World Browser: - Added tiled landscape importing

diff --git a/Engine/Source/Editor/WorldBrowser/Private/LevelCollectionCommands.h b/Engine/Source/Editor/WorldBrowser/Private/LevelCollectionCommands.h
--- a/Engine/Source/Editor/WorldBrowser/Private/LevelCollectionCommands.h
+++ b/Engine/Source/Editor/WorldBrowser/Private/LevelCollectionCommands.h
@@ -105,6 +105,7 @@ public:
 		UI_COMMAND( MoveLevelDown,	"Move Level Down",		"Moves level down by 1 unit", EUserInterfaceActionType::Button, FInputGesture(EKeys::Down) );
 		
 		// Landscape operations
+		UI_COMMAND( ImportTiledLandscape, "Import Tiled Landscape...", "Imports landscape from a tiled heightmap files (<name>X<n>_Y<n>.r16)", EUserInterfaceActionType::Button, FInputGesture() );
 		UI_COMMAND( AddLandscapeLevelXNegative,	"-X",	"Add a new adjacent level with landscape proxy", EUserInterfaceActionType::Button, FInputGesture() );
 		UI_COMMAND( AddLandscapeLevelXPositive,	"+X",	"Add a new adjacent level with landscape proxy", EUserInterfaceActionType::Button, FInputGesture() );
 		UI_COMMAND( AddLandscapeLevelYNegative,	"-Y",	"Add a new adjacent level with landscape proxy", EUserInterfaceActionType::Button, FInputGesture() );
@@ -284,6 +285,8 @@ public:
 
 
 	// Landscape operations
+	TSharedPtr< FUICommandInfo > ImportTiledLandscape;
+	
 	TSharedPtr< FUICommandInfo > AddLandscapeLevelXNegative;
 	TSharedPtr< FUICommandInfo > AddLandscapeLevelXPositive;
 	TSharedPtr< FUICommandInfo > AddLandscapeLevelYNegative;
diff --git a/Engine/Source/Editor/WorldBrowser/Private/LevelCollectionModel.h b/Engine/Source/Editor/WorldBrowser/Private/LevelCollectionModel.h
--- a/Engine/Source/Editor/WorldBrowser/Private/LevelCollectionModel.h
+++ b/Engine/Source/Editor/WorldBrowser/Private/LevelCollectionModel.h
@@ -510,4 +510,26 @@ struct FUnmodifiableObject
 private:
 	UObject*		ImmuneObject;
 	bool			bTransient;
+};
+
+/**  */
+struct FTiledLandscapeImportSettings
+{
+	FTiledLandscapeImportSettings()
+		: Scale3D(100.f,100.f,100.f)
+		, ComponentsNum(8)
+		, QuadsPerSection(63)
+		, SectionsPerComponent(1)
+		, TileResolution(1009)	
+	{}
+	
+	FVector			Scale3D;
+	int32			ComponentsNum;
+	int32			QuadsPerSection;
+	int32			SectionsPerComponent;
+	
+	TArray<FString>		ImportFileList;
+	TArray<FIntPoint>	TileCoordinates;
+	FIntRect			TotalLandscapeRect;	
+	int32				TileResolution;	
 };
\ No newline at end of file
diff --git a/Engine/Source/Editor/WorldBrowser/Private/Tiles/STiledLandscapeImportDlg.cpp b/Engine/Source/Editor/WorldBrowser/Private/Tiles/STiledLandscapeImportDlg.cpp
new file mode 100644
--- /dev/null
+++ b/Engine/Source/Editor/WorldBrowser/Private/Tiles/STiledLandscapeImportDlg.cpp
@@ -0,0 +1,412 @@
+// Copyright 1998-2014 Epic Games, Inc. All Rights Reserved.
+
+#include "WorldBrowserPrivatePCH.h"
+#include "STiledLandscapeImportDlg.h"
+
+#define LOCTEXT_NAMESPACE "WorldBrowser"
+
+static int32 CalcLandscapeSquareResolution(int32 ComponetsNumX, int32 SectionNumX, int32 SectionQuadsNumX)
+{
+	return ComponetsNumX*SectionNumX*SectionQuadsNumX+1;
+}
+
+/**
+ *	Returns heigtmap tile coordinates extracted from a specified tile filename
+ */
+static FIntPoint ExtractHeighmapTileCoordinates(FString BaseFilename)
+{
+	//We expect file name in form: <tilename>_x<number>_y<number>
+		
+	FIntPoint ResultPosition(-1,-1);
+	
+	int32 XPos = BaseFilename.Find(TEXT("_x"), ESearchCase::IgnoreCase, ESearchDir::FromEnd);
+	int32 YPos = BaseFilename.Find(TEXT("_y"), ESearchCase::IgnoreCase, ESearchDir::FromEnd);
+	if (XPos != INDEX_NONE && YPos != INDEX_NONE && XPos < YPos)
+	{
+		FString XCoord = BaseFilename.Mid(XPos+2, YPos-(XPos+2));
+		FString YCoord = BaseFilename.Mid(YPos+2, BaseFilename.Len()-(YPos+2));
+
+		if (XCoord.IsNumeric() && YCoord.IsNumeric())
+		{
+			TTypeFromString<int32>::FromString(ResultPosition.X, *XCoord);
+			TTypeFromString<int32>::FromString(ResultPosition.Y, *YCoord);
+		}
+	}
+
+	return ResultPosition;
+}
+
+void STiledLandcapeImportDlg::Construct(const FArguments& InArgs, TSharedPtr<SWindow> InParentWindow)
+{
+	ParentWindow = InParentWindow;
+	
+	ChildSlot
+	[
+		SNew(SBorder)
+		.BorderImage(FEditorStyle::GetBrush("ToolPanel.GroupBorder"))
+		[
+			SNew(SVerticalBox)
+
+			+SVerticalBox::Slot()
+			.FillHeight(1)
+			.Padding(0,10,0,10)
+			[
+				SNew(SUniformGridPanel)
+				.SlotPadding(2)
+				
+				// Select tiles
+				+SUniformGridPanel::Slot(0, 0)
+
+				+SUniformGridPanel::Slot(1, 0)
+				.VAlign(VAlign_Center)
+				[
+					SNew(SButton)
+					.HAlign(HAlign_Center)
+					.ContentPadding(FEditorStyle::GetMargin("StandardDialog.ContentPadding"))
+					.OnClicked(this, &STiledLandcapeImportDlg::OnClickedSelectTiles)
+					.Text(LOCTEXT("TiledLandscapeImport_SelectButtonText", "Select Tiles..."))
+				]
+
+				// Tile configuration
+				+SUniformGridPanel::Slot(0, 1)
+				.VAlign(VAlign_Center)
+				[
+					SNew(STextBlock)
+					.Text(LOCTEXT("TiledLandscapeImport_ConfigurationText", "Import Configuration"))
+				]
+
+				+SUniformGridPanel::Slot(1, 1)
+				.VAlign(VAlign_Center)
+				[
+					SAssignNew(TileConfigurationComboBox, SComboBox<TSharedPtr<FTileImportConfiguration>>)
+					.OptionsSource(&ActiveConfigurations)
+					.OnSelectionChanged(this, &STiledLandcapeImportDlg::OnSetImportConfiguration)
+					.OnGenerateWidget(this, &STiledLandcapeImportDlg::HandleTileConfigurationComboBoxGenarateWidget)
+					.Content()
+					[
+						SNew(STextBlock)
+						.Text(this, &STiledLandcapeImportDlg::GetTileConfigurationText)
+					]
+				]
+
+				// Scale
+				+SUniformGridPanel::Slot(0, 2)
+				.VAlign(VAlign_Center)
+				[
+					SNew(STextBlock)
+					.Text(LOCTEXT("TiledLandscapeImport_ScaleText", "Landscape Scale"))
+				]
+			
+				+SUniformGridPanel::Slot(1, 2)
+				.VAlign(VAlign_Center)
+				[
+					SNew( SVectorInputBox )
+					.bColorAxisLabels( true )
+					.X( this, &STiledLandcapeImportDlg::GetScaleX )
+					.Y( this, &STiledLandcapeImportDlg::GetScaleY )
+					.Z( this, &STiledLandcapeImportDlg::GetScaleZ )
+					.OnXCommitted( this, &STiledLandcapeImportDlg::OnSetScale, 0 )
+					.OnYCommitted( this, &STiledLandcapeImportDlg::OnSetScale, 1 )
+					.OnZCommitted( this, &STiledLandcapeImportDlg::OnSetScale, 2 )
+				]
+
+			]
+
+			// Import summary
+			+SVerticalBox::Slot()
+			.AutoHeight()
+			.HAlign(HAlign_Right)
+			.VAlign(VAlign_Center)
+			.Padding(FEditorStyle::GetMargin("StandardDialog.ContentPadding"))
+			[
+				SNew(STextBlock)
+				.Text(this, &STiledLandcapeImportDlg::GetImportSummaryText)
+			]
+
+			// Import, Cancel
+			+SVerticalBox::Slot()
+			.AutoHeight()
+			.HAlign(HAlign_Right)
+			.VAlign(VAlign_Bottom)
+			.Padding(0,10,0,10)
+			[
+				SNew(SUniformGridPanel)
+				.SlotPadding(FEditorStyle::GetMargin("StandardDialog.SlotPadding"))
+				.MinDesiredSlotWidth(FEditorStyle::GetFloat("StandardDialog.MinDesiredSlotWidth"))
+				.MinDesiredSlotHeight(FEditorStyle::GetFloat("StandardDialog.MinDesiredSlotHeight"))
+				+SUniformGridPanel::Slot(0,0)
+				[
+					SNew(SButton)
+					.HAlign(HAlign_Center)
+					.ContentPadding(FEditorStyle::GetMargin("StandardDialog.ContentPadding"))
+					.IsEnabled(this, &STiledLandcapeImportDlg::IsImportEnabled)
+					.OnClicked(this, &STiledLandcapeImportDlg::OnClickedImport)
+					.Text(LOCTEXT("TiledLandscapeImport_ImportButtonText", "Import"))
+				]
+				+SUniformGridPanel::Slot(1,0)
+				[
+					SNew(SButton)
+					.HAlign(HAlign_Center)
+					.ContentPadding(FEditorStyle::GetMargin("StandardDialog.ContentPadding"))
+					.OnClicked(this, &STiledLandcapeImportDlg::OnClickedCancel)
+					.Text(LOCTEXT("TiledLandscapeImport_CancelButtonText", "Cancel"))
+				]
+			]
+		]
+	];
+
+
+	GenerateAllPossibleTileConfigurations();
+	SetPossibleConfigurationsForResolution(0);
+}
+
+TSharedRef<SWidget> STiledLandcapeImportDlg::HandleTileConfigurationComboBoxGenarateWidget(TSharedPtr<FTileImportConfiguration> InItem) const
+{
+	const FText ItemText = GenerateConfigurationText(InItem->NumComponents, InItem->NumSectionsPerComponent, InItem->NumQuadsPerSection);
+		
+	return SNew(SBox)
+	.Padding(4)
+	[
+		SNew(STextBlock).Text(ItemText)
+	];
+}
+
+FText STiledLandcapeImportDlg::GetTileConfigurationText() const
+{
+	if (ImportSettings.ImportFileList.Num() == 0)
+	{
+		return LOCTEXT("TiledLandscapeImport_NoTilesText", "No tiles selected");
+	}
+
+	if (ImportSettings.SectionsPerComponent <= 0)
+	{
+		return LOCTEXT("TiledLandscapeImport_InvalidTileResolutuinText", "Selected tiles have unsupported resolutuion");
+	}
+	
+	return GenerateConfigurationText(ImportSettings.ComponentsNum, ImportSettings.SectionsPerComponent, ImportSettings.QuadsPerSection);
+}
+
+const FTiledLandscapeImportSettings& STiledLandcapeImportDlg::GetImportSettings() const
+{
+	return ImportSettings;
+}
+
+TOptional<float> STiledLandcapeImportDlg::GetScaleX() const
+{
+	return ImportSettings.Scale3D.X;
+}
+
+TOptional<float> STiledLandcapeImportDlg::GetScaleY() const
+{
+	return ImportSettings.Scale3D.Y;
+}
+
+TOptional<float> STiledLandcapeImportDlg::GetScaleZ() const
+{
+	return ImportSettings.Scale3D.Z;
+}
+
+void STiledLandcapeImportDlg::OnSetScale(float InValue, ETextCommit::Type CommitType, int32 InAxis)
+{
+	if (InAxis < 2) //XY uniform
+	{
+		ImportSettings.Scale3D.X = FMath::Abs(InValue);
+		ImportSettings.Scale3D.Y = FMath::Abs(InValue);
+	}
+	else //Z
+	{
+		ImportSettings.Scale3D.Z = FMath::Abs(InValue);
+	}
+}
+
+void STiledLandcapeImportDlg::OnSetImportConfiguration(TSharedPtr<FTileImportConfiguration> InTileConfig, ESelectInfo::Type SelectInfo)
+{
+	if (InTileConfig.IsValid())
+	{
+		ImportSettings.ComponentsNum = InTileConfig->NumComponents;
+		ImportSettings.QuadsPerSection = InTileConfig->NumQuadsPerSection;
+		ImportSettings.SectionsPerComponent = InTileConfig->NumSectionsPerComponent;
+	}
+	else
+	{
+		ImportSettings.ComponentsNum = 0;
+		ImportSettings.ImportFileList.Empty();
+	}
+}
+
+FReply STiledLandcapeImportDlg::OnClickedSelectTiles()
+{
+	ImportSettings.ImportFileList.Empty();
+	ImportSettings.TileCoordinates.Empty();
+	ImportSettings.TotalLandscapeRect = FIntRect(MAX_int32, MAX_int32, MIN_int32, MIN_int32);
+
+	SetPossibleConfigurationsForResolution(0);
+	
+	IDesktopPlatform* DesktopPlatform = FDesktopPlatformModule::Get();
+	if (DesktopPlatform)
+	{
+		if (ParentWindow->GetNativeWindow().IsValid())
+		{
+			void* ParentWindowWindowHandle = ParentWindow->GetNativeWindow()->GetOSWindowHandle();
+
+			const FString ImportFileTypes = TEXT("Raw heightmap tiles (*.r16)|*.r16");
+			if (DesktopPlatform->OpenFileDialog(
+				ParentWindowWindowHandle,
+				LOCTEXT("SelectHeightmapTiles", "Select heightmap tiles").ToString(),
+				*FEditorDirectories::Get().GetLastDirectory(ELastDirectory::UNR),
+				TEXT(""),
+				*ImportFileTypes,
+				EFileDialogFlags::Multiple,
+				ImportSettings.ImportFileList
+				) && ImportSettings.ImportFileList.Num())
+			{
+				IFileManager& FileManager = IFileManager::Get();
+				bool bValidTiles = true;
+
+				// All heightmap tiles have to be the same size and have correct tile position encoded into filename
+				const int64 TargetFileSize = FileManager.FileSize(*ImportSettings.ImportFileList[0]);
+				for (const FString& Filename : ImportSettings.ImportFileList)
+				{
+					int64 FileSize = FileManager.FileSize(*Filename);
+					if (FileSize != TargetFileSize)
+					{
+						bValidTiles = false;
+						break;
+					}
+
+					FIntPoint TileCoordinate = ExtractHeighmapTileCoordinates(FPaths::GetBaseFilename(Filename));
+					if (TileCoordinate.GetMin() < 0)
+					{
+						bValidTiles = false;
+						break;
+					}
+
+					ImportSettings.TileCoordinates.Add(TileCoordinate);
+					ImportSettings.TotalLandscapeRect.Include(TileCoordinate);
+				}
+				
+				if (bValidTiles)
+				{
+					ImportSettings.TileResolution = FMath::Sqrt(TargetFileSize/2.0);
+					SetPossibleConfigurationsForResolution(ImportSettings.TileResolution);
+				}
+			}
+		}
+	}
+		
+	return FReply::Handled();
+}
+
+bool STiledLandcapeImportDlg::IsImportEnabled() const
+{
+	return ImportSettings.ImportFileList.Num() && ImportSettings.ComponentsNum > 0;
+}
+
+FReply STiledLandcapeImportDlg::OnClickedImport()
+{
+	ParentWindow->RequestDestroyWindow();
+	return FReply::Handled();	
+}
+
+FReply STiledLandcapeImportDlg::OnClickedCancel()
+{
+	ParentWindow->RequestDestroyWindow();
+	
+	ImportSettings.ImportFileList.Empty();
+	return FReply::Handled();	
+}
+
+void STiledLandcapeImportDlg::SetPossibleConfigurationsForResolution(int32 TargetResolutuion)
+{
+	int32 Idx = AllConfigurations.IndexOfByPredicate([&](const FTileImportConfiguration& A){
+		return TargetResolutuion == A.Resolutuion;
+	});
+
+	ActiveConfigurations.Empty();
+	ImportSettings.ComponentsNum = 0; // Set invalid options
+
+	// AllConfigurations is sorted by resolution
+	while(AllConfigurations.IsValidIndex(Idx) && AllConfigurations[Idx].Resolutuion == TargetResolutuion)
+	{
+		TSharedPtr<FTileImportConfiguration> TileConfig = MakeShareable(new FTileImportConfiguration(AllConfigurations[Idx++]));
+		ActiveConfigurations.Add(TileConfig);
+	}
+
+	// Refresh combo box with active configurations
+	TileConfigurationComboBox->RefreshOptions();
+	// Set first configurationion as active
+	if (ActiveConfigurations.Num())
+	{
+		TileConfigurationComboBox->SetSelectedItem(ActiveConfigurations[0]);
+	}
+}
+
+void STiledLandcapeImportDlg::GenerateAllPossibleTileConfigurations()
+{
+	AllConfigurations.Empty();
+	for (int32 ComponentsNum = 1; ComponentsNum <= 32; ComponentsNum++)
+	{
+		for (int32 SectionsPerComponent = 1; SectionsPerComponent <= 2; SectionsPerComponent++)
+		{
+			for (int32 QuadsPerSection = 3; QuadsPerSection <= 8; QuadsPerSection++)
+			{
+				FTileImportConfiguration Entry;
+				Entry.NumComponents				= ComponentsNum;
+				Entry.NumSectionsPerComponent	= SectionsPerComponent;
+				Entry.NumQuadsPerSection		= (1 << QuadsPerSection) - 1;
+				Entry.Resolutuion				= CalcLandscapeSquareResolution(Entry.NumComponents, Entry.NumSectionsPerComponent, Entry.NumQuadsPerSection);
+
+				AllConfigurations.Add(Entry);
+			}
+		}
+	}
+	
+	// Sort by resolution
+	AllConfigurations.Sort([](const FTileImportConfiguration& A, const FTileImportConfiguration& B){
+		if (A.Resolutuion == B.Resolutuion)
+		{
+			return A.NumComponents < B.NumComponents;
+		}
+		return A.Resolutuion < B.Resolutuion;
+	});
+}
+
+FText STiledLandcapeImportDlg::GetImportSummaryText() const
+{
+	if (ImportSettings.ImportFileList.Num() == 0 || ImportSettings.ComponentsNum <= 0)
+	{
+		return FText();
+	}
+	
+	// Tile information(num, resolution)
+	const FString TilesSummary = FString::Printf(TEXT("%d - %dx%d"), ImportSettings.ImportFileList.Num(), ImportSettings.TileResolution, ImportSettings.TileResolution);
+	
+	// Total landscape size(NxN km)
+	int32 WidthInTilesX = ImportSettings.TotalLandscapeRect.Width() + 1;
+	int32 WidthInTilesY = ImportSettings.TotalLandscapeRect.Height() + 1;
+	float WidthX = 0.00001f*ImportSettings.Scale3D.X*WidthInTilesX*ImportSettings.TileResolution;
+	float WidthY = 0.00001f*ImportSettings.Scale3D.Y*WidthInTilesY*ImportSettings.TileResolution;
+	const FString LandscapeSummary = FString::Printf(TEXT("%.3fx%.3f"), WidthX, WidthY);
+	
+	return FText::Format(
+		LOCTEXT("TiledLandscapeImport_SummaryText", "{0} tiles, {1}km landscape"), 
+		FText::FromString(TilesSummary),
+		FText::FromString(LandscapeSummary)
+		);
+}
+
+FText STiledLandcapeImportDlg::GenerateConfigurationText(int32 NumComponents, int32 NumSectionsPerComponent, int32 NumQuadsPerSection) const
+{
+	const FString ComponentsStr = FString::Printf(TEXT("%dx%d"), NumComponents, NumComponents);
+	const FString SectionsStr = FString::Printf(TEXT("%dx%d"), NumSectionsPerComponent, NumSectionsPerComponent);
+	const FString QuadsStr = FString::Printf(TEXT("%dx%d"), NumQuadsPerSection, NumQuadsPerSection);
+	
+	return FText::Format(
+		LOCTEXT("TiledLandscapeImport_ConfigurationText", "Components: {0} Sections: {1} Quads: {2}"), 
+		FText::FromString(ComponentsStr),
+		FText::FromString(SectionsStr),
+		FText::FromString(QuadsStr)
+		);
+}
+
+#undef LOCTEXT_NAMESPACE
diff --git a/Engine/Source/Editor/WorldBrowser/Private/Tiles/STiledLandscapeImportDlg.h b/Engine/Source/Editor/WorldBrowser/Private/Tiles/STiledLandscapeImportDlg.h
new file mode 100644
--- /dev/null
+++ b/Engine/Source/Editor/WorldBrowser/Private/Tiles/STiledLandscapeImportDlg.h
@@ -0,0 +1,77 @@
+// Copyright 1998-2014 Epic Games, Inc. All Rights Reserved.
+#pragma once
+
+/////////////////////////////////////////////////////
+// STiledLandcapeImportDlg
+// 
+class STiledLandcapeImportDlg
+	: public SCompoundWidget
+{
+public:	
+	/** */		
+	struct FTileImportConfiguration
+	{
+		int32 Resolutuion;
+		
+		int32 NumComponents;
+		int32 NumSectionsPerComponent;
+		int32 NumQuadsPerSection;
+	};
+
+
+	SLATE_BEGIN_ARGS( STiledLandcapeImportDlg )
+		{}
+	SLATE_END_ARGS()
+	
+	void Construct(const FArguments& InArgs, TSharedPtr<SWindow> InParentWindow);
+
+	/** */
+	const FTiledLandscapeImportSettings& GetImportSettings() const;
+
+private:
+	TSharedRef<SWidget> HandleTileConfigurationComboBoxGenarateWidget(TSharedPtr<FTileImportConfiguration> InItem) const;
+	FText				GetTileConfigurationText() const;
+
+	/** */
+	TOptional<float> GetScaleX() const;
+	TOptional<float> GetScaleY() const;
+	TOptional<float> GetScaleZ() const;
+	void OnSetScale(float InValue, ETextCommit::Type CommitType, int32 InAxis);
+
+	/** */
+	void OnSetImportConfiguration(TSharedPtr<FTileImportConfiguration> InTileConfig, ESelectInfo::Type SelectInfo);
+
+	/** */
+	bool IsImportEnabled() const;
+		
+	/** */
+	FReply OnClickedSelectTiles();
+	FReply OnClickedImport();
+	FReply OnClickedCancel();
+
+	/** */
+	FText GetImportSummaryText() const;
+
+	/** */
+	void SetPossibleConfigurationsForResolution(int32 TargetResolutuion);
+	
+	/** */
+	void GenerateAllPossibleTileConfigurations();
+
+	/** */
+	FText GenerateConfigurationText(int32 NumComponents, int32 NumSectionsPerComponent,	int32 NumQuadsPerSection) const;
+
+private:
+	/** */
+	TSharedPtr<SWindow> ParentWindow;
+	
+	/** */
+	TSharedPtr<SComboBox<TSharedPtr<FTileImportConfiguration>>> TileConfigurationComboBox;
+	
+	/** */
+	FTiledLandscapeImportSettings	ImportSettings;
+	
+	TArray<FTileImportConfiguration> AllConfigurations;
+	TArray<TSharedPtr<FTileImportConfiguration>> ActiveConfigurations;
+};
+
diff --git a/Engine/Source/Editor/WorldBrowser/Private/Tiles/WorldTileCollectionModel.cpp b/Engine/Source/Editor/WorldBrowser/Private/Tiles/WorldTileCollectionModel.cpp
--- a/Engine/Source/Editor/WorldBrowser/Private/Tiles/WorldTileCollectionModel.cpp
+++ b/Engine/Source/Editor/WorldBrowser/Private/Tiles/WorldTileCollectionModel.cpp
@@ -13,6 +13,7 @@
 #include "WorldTileDetails.h"
 #include "WorldTileDetailsCustomization.h"
 #include "WorldTileCollectionModel.h"
+#include "STiledLandscapeImportDlg.h"
 
 #define LOCTEXT_NAMESPACE "WorldBrowser"
 
@@ -426,6 +427,7 @@ void FWorldTileCollectionModel::CustomizeFileMainMenu(FMenuBuilder& InMenuBuilde
 	InMenuBuilder.BeginSection("LevelsAddLevel");
 	{
 		InMenuBuilder.AddMenuEntry( Commands.World_CreateEmptyLevel );
+		InMenuBuilder.AddMenuEntry( Commands.ImportTiledLandscape );
 	}
 	InMenuBuilder.EndSection();
 }
@@ -551,6 +553,9 @@ void FWorldTileCollectionModel::BindCommands()
 
 
 	// Landscape operations
+	ActionList.MapAction(Commands.ImportTiledLandscape,
+		FExecuteAction::CreateSP(this, &FWorldTileCollectionModel::ImportTiledLandscape_Executed));
+
 	ActionList.MapAction(Commands.AddLandscapeLevelXNegative,
 		FExecuteAction::CreateSP(this, &FWorldTileCollectionModel::AddLandscapeProxy_Executed, FWorldTileModel::XNegative),
 		FCanExecuteAction::CreateSP(this, &FWorldTileCollectionModel::CanAddLandscapeProxy, FWorldTileModel::XNegative));
@@ -1230,6 +1235,146 @@ void FWorldTileCollectionModel::AddLandscapeProxy_Executed(FWorldTileModel::EWor
 	}
 }
 
+template<typename DataType>
+static bool ReadHeighmapFile(TArray<DataType>& Result, const TCHAR* Filename, uint32 Flags = 0)
+{
+	auto Reader = TScopedPointer<FArchive>(IFileManager::Get().CreateFileReader(Filename, Flags));
+	if (!Reader)
+	{
+		if (!(Flags & FILEREAD_Silent))
+		{
+			UE_LOG(LogStreaming, Warning, TEXT("Failed to read file '%s' error."), Filename);
+		}
+		return 0;
+	}
+	Result.Reset();
+	Result.AddUninitialized(Reader->TotalSize()/Result.GetTypeSize());
+	Reader->Serialize(Result.GetTypedData(), Result.Num()*Result.GetTypeSize());
+	return Reader->Close();
+}
+
+void FWorldTileCollectionModel::ImportTiledLandscape_Executed()
+{
+	/** Create the window to host widget */
+	TSharedRef<SWindow> ImportWidnow = SNew(SWindow)
+											.Title(LOCTEXT("TiledLandcapeImport_DialogTitle", "Import Tiled Heightmap"))
+											.SizingRule( ESizingRule::Autosized )
+											.SupportsMinimize(false) 
+											.SupportsMaximize(false);
+
+	/** Set the content of the window */
+	TSharedRef<STiledLandcapeImportDlg> ImportDialog = SNew(STiledLandcapeImportDlg, ImportWidnow);
+	ImportWidnow->SetContent(ImportDialog);
+
+	/** Show the dialog window as a modal window */
+	Editor->EditorAddModalWindow(ImportWidnow);
+
+	const FTiledLandscapeImportSettings& ImportSettings = ImportDialog->GetImportSettings();
+		
+	if (ImportSettings.ImportFileList.Num())
+	{
+		// Default path for imported landscape tiles
+		// Use tile prefix as a folder name under world root
+		FString WorldRootPath = FPackageName::LongPackageNameToFilename(GetWorld()->WorldComposition->GetWorldRoot());
+		// Extract tile prefix
+		FString FolderName = FPaths::GetBaseFilename(ImportSettings.ImportFileList[0]);
+		int32 PrefixEnd = FolderName.Find(TEXT("_x"), ESearchCase::IgnoreCase, ESearchDir::FromEnd);
+		FolderName = FolderName.Left(PrefixEnd);
+		WorldRootPath+= FolderName;
+		WorldRootPath+= TEXT("/");
+
+		GWarn->BeginSlowTask(LOCTEXT("ImportingLandscapeTilesBegin", "Importing landscape tiles"), true);
+		
+		// Shared GUID used between landscape proxies
+		FGuid LandscapeGuid;
+
+		// Import tiles
+		int32 TileIndex = 0;
+		for (const FString& Filename : ImportSettings.ImportFileList)
+		{
+			FString TileName = FPaths::GetBaseFilename(Filename);
+			FVector TileScale = ImportSettings.Scale3D;
+
+			GWarn->StatusUpdate(TileIndex, ImportSettings.ImportFileList.Num(), FText::Format(LOCTEXT("ImportingLandscapeTiles", "Importing landscape tiles: {0}"), FText::FromString(TileName)));
+			
+			FWorldTileModel::FLandscapeImportSettings TileImportSettings;
+			TileImportSettings.LandscapeGuid		= LandscapeGuid;
+			TileImportSettings.ComponentSizeQuads	= ImportSettings.QuadsPerSection*ImportSettings.SectionsPerComponent;
+			TileImportSettings.QuadsPerSection		= ImportSettings.QuadsPerSection;
+			TileImportSettings.SectionsPerComponent = ImportSettings.SectionsPerComponent;
+			TileImportSettings.SizeX				= ImportSettings.TileResolution;
+			TileImportSettings.SizeY				= ImportSettings.TileResolution;
+			TileImportSettings.LandscapeTransform.SetScale3D(TileScale);
+
+			if (ReadHeighmapFile(TileImportSettings.HeightData, *Filename, FILEREAD_Silent))
+			{
+				FString MapFileName = WorldRootPath + TileName + FPackageName::GetMapPackageExtension();
+				// Create a new world - so we can 'borrow' its level
+				UWorld* NewWorld = UWorld::CreateWorld(EWorldType::None, false);
+				check(NewWorld);
+								
+				bool bSaved = FEditorFileUtils::SaveLevel(NewWorld->PersistentLevel, *MapFileName);
+				if (bSaved)
+				{
+					// update levels list so we can find a new level in our world model
+					PopulateLevelsList();
+					TSharedPtr<FWorldTileModel> NewTileModel = StaticCastSharedPtr<FWorldTileModel>(FindLevelModel(NewWorld->GetOutermost()->GetFName()));
+					// Hide level, so we do not depend on a current world origin
+					NewTileModel->SetVisible(false);
+					
+					// Create landscape (actor/proxy) in a new level
+					ALandscapeProxy* NewLandscape = NewTileModel->ImportLandscape(TileImportSettings);
+
+					// Use first created landscape as a landscape actor and all next tiles as landscape proxies 
+					if (!LandscapeGuid.IsValid())
+					{
+						LandscapeGuid = NewLandscape->GetLandscapeGuid();
+					}
+
+					if (NewLandscape)
+					{
+						// Set landscape actor location at min corner of bounding rect, so it will be centered around origin
+						FIntRect NewLandscapeRect = NewLandscape->GetBoundingRect();
+						float WidthX = NewLandscapeRect.Width()*TileScale.X;
+						float WidthY = NewLandscapeRect.Height()*TileScale.Y;
+						NewLandscape->SetActorLocation(-FVector(WidthX, WidthY, 0.f)*0.5f);
+						
+						// Set bounds of a tile
+						NewTileModel->TileDetails->Bounds = NewLandscape->GetComponentsBoundingBox();
+
+						// Calculate this tile offset from world origin
+						FIntPoint TileCoordinates = ImportSettings.TileCoordinates[TileIndex];
+						
+						// Mirror Y coordinate, to match unreal coordinate system
+						// TileCoordinates.Y = TilesRect.Height() - TileCoordinates.Y; //option
+
+						// Tiles coordinates are in range [0..N], convert them to [-N/2..N/2] so whole landscape will be centered around world origin
+						int32 WidthInTilesX = ImportSettings.TotalLandscapeRect.Width() + 1;
+						int32 WidthInTilesY = ImportSettings.TotalLandscapeRect.Height() + 1;
+						TileCoordinates = TileCoordinates - FIntPoint(WidthInTilesX/2, WidthInTilesY/2);
+												
+						FIntPoint TileOffset = FIntPoint(TileCoordinates.X*WidthX, TileCoordinates.Y*WidthY);
+						
+						// Place level tile at correct position in the world
+						NewTileModel->SetLevelPosition(TileOffset);
+					
+						// Save level with a landscape
+						FEditorFileUtils::SaveLevel(NewWorld->PersistentLevel, *MapFileName);
+					}
+					
+					// Destroy the new world we created and collect the garbage
+					NewWorld->DestroyWorld(false);
+					CollectGarbage(GARBAGE_COLLECTION_KEEPFLAGS);
+				}
+			}
+
+			TileIndex++;
+		}
+
+		GWarn->EndSlowTask();
+	}
+}
+
 void FWorldTileCollectionModel::PostUndo(bool bSuccess)
 {
 	if (!bIsSavingLevel)
diff --git a/Engine/Source/Editor/WorldBrowser/Private/Tiles/WorldTileCollectionModel.h b/Engine/Source/Editor/WorldBrowser/Private/Tiles/WorldTileCollectionModel.h
--- a/Engine/Source/Editor/WorldBrowser/Private/Tiles/WorldTileCollectionModel.h
+++ b/Engine/Source/Editor/WorldBrowser/Private/Tiles/WorldTileCollectionModel.h
@@ -211,6 +211,9 @@ private:
 	/** @return whether it is possible to add a level with landscape proxy at specified location */	
 	bool CanAddLandscapeProxy(FWorldTileModel::EWorldDirections InWhere) const;
 
+	/**  */	
+	void ImportTiledLandscape_Executed();
+
 public:
 	/** Whether Editor has support for generating LOD levels */	
 	bool HasGenerateLODLevelSupport() const;
diff --git a/Engine/Source/Editor/WorldBrowser/Private/Tiles/WorldTileModel.cpp b/Engine/Source/Editor/WorldBrowser/Private/Tiles/WorldTileModel.cpp
--- a/Engine/Source/Editor/WorldBrowser/Private/Tiles/WorldTileModel.cpp
+++ b/Engine/Source/Editor/WorldBrowser/Private/Tiles/WorldTileModel.cpp
@@ -834,7 +834,7 @@ bool FWorldTileModel::CreateAdjacentLandscapeProxy(ALandscapeProxy* SourceLandsc
 	FIntPoint SourceLandscapeSize = SourceLandscapeRect.Size();
 
 	FLandscapeImportSettings ImportSettings;
-	ImportSettings.SourceLandscape = SourceLandscape;
+	ImportSettings.LandscapeGuid = SourceLandscape->GetLandscapeGuid();
 	ImportSettings.ComponentSizeQuads = SourceLandscape->ComponentSizeQuads;
 	ImportSettings.SectionsPerComponent = SourceLandscape->NumSubsections;
 	ImportSettings.QuadsPerSection = SourceLandscape->SubsectionSizeQuads;
@@ -856,6 +856,9 @@ bool FWorldTileModel::CreateAdjacentLandscapeProxy(ALandscapeProxy* SourceLandsc
 	ALandscapeProxy* AdjacenLandscape = ImportLandscape(ImportSettings);
 	if (AdjacenLandscape)
 	{
+		// Copy source landscape properties 
+		AdjacenLandscape->GetSharedProperties(SourceLandscape);
+		
 		// Refresh level model bounding box
 		FBox AdjacentLandscapeBounds = AdjacenLandscape->GetComponentsBoundingBox(true);
 		TileDetails->Bounds = AdjacentLandscapeBounds;
@@ -898,23 +901,12 @@ ALandscapeProxy* FWorldTileModel::ImportLandscape(const FLandscapeImportSettings
 		return nullptr;
 	}
 	
-	// In case SourceLandscape os provided, create LnadscapeProxy and copy settings from it
-	// Otherwise create a new LandscapeActor
 	ALandscapeProxy*	Landscape;
-	FGuid				LandscapeGuid;
-	if (Settings.SourceLandscape)
+	FGuid				LandscapeGuid = Settings.LandscapeGuid;
+	
+	if (LandscapeGuid.IsValid())
 	{
-		// These settings have to match source landscape settings
-		if (Settings.ComponentSizeQuads != Settings.SourceLandscape->ComponentSizeQuads || 
-			Settings.SectionsPerComponent != Settings.SourceLandscape->NumSubsections ||
-			Settings.QuadsPerSection != Settings.SourceLandscape->SubsectionSizeQuads)
-		{
-			return nullptr;
-		}
-				
 		Landscape = Cast<UWorld>(LoadedLevel->GetOuter())->SpawnActor<ALandscapeProxy>();
-		Landscape->GetSharedProperties(Settings.SourceLandscape);
-		LandscapeGuid = Settings.SourceLandscape->GetLandscapeGuid();
 	}
 	else
 	{
diff --git a/Engine/Source/Editor/WorldBrowser/Private/Tiles/WorldTileModel.h b/Engine/Source/Editor/WorldBrowser/Private/Tiles/WorldTileModel.h
--- a/Engine/Source/Editor/WorldBrowser/Private/Tiles/WorldTileModel.h
+++ b/Engine/Source/Editor/WorldBrowser/Private/Tiles/WorldTileModel.h
@@ -36,7 +36,8 @@ public:
 	 */
 	struct FLandscapeImportSettings
 	{
-		ALandscapeProxy*					SourceLandscape;
+		// Depending on landscape guid import code will spawn Landscape actor or LandscapeProxy actor
+		FGuid								LandscapeGuid;
 		FTransform							LandscapeTransform;
 		int32								ComponentSizeQuads;
 		int32								SectionsPerComponent;
