World Browser:
- Fixed positioning issue with creating adjacent landcapes
- Fixed issue with translating multiple landscape levels at once

diff --git a/Engine/Source/Editor/WorldBrowser/Private/Tiles/WorldTileModel.cpp b/Engine/Source/Editor/WorldBrowser/Private/Tiles/WorldTileModel.cpp
--- a/Engine/Source/Editor/WorldBrowser/Private/Tiles/WorldTileModel.cpp
+++ b/Engine/Source/Editor/WorldBrowser/Private/Tiles/WorldTileModel.cpp
@@ -504,7 +504,7 @@ void FWorldTileModel::SetLevelPosition(const FIntPoint& InPosition)
 void FWorldTileModel::FixLandscapeSectionsOffset()
 {
 	check(IsLandscapeBased());
-	ULandscapeInfo::RecreateLandscapeInfo(LevelCollectionModel.GetWorld(), true);
+	Editor->DeferredCommands.AddUnique(TEXT("UpdateLandscapeEditorData"));
 }
 
 void FWorldTileModel::Update()
@@ -884,7 +884,7 @@ bool FWorldTileModel::CreateAdjacentLandscapeProxy(ALandscapeProxy* SourceLandsc
 		}
 
 		// Add source level position
-		FIntPoint IntOffset = FIntPoint(ProxyOffset.X, ProxyOffset.Y) + SourceTileOffset;
+		FIntPoint IntOffset = FIntPoint(ProxyOffset.X, ProxyOffset.Y) + LevelCollectionModel.GetWorld()->GlobalOriginOffset;
 
 		// Move level with landscape proxy to desired position
 		SetLevelPosition(IntOffset);
diff --git a/Engine/Source/Runtime/Engine/Private/Landscape/Landscape.cpp b/Engine/Source/Runtime/Engine/Private/Landscape/Landscape.cpp
--- a/Engine/Source/Runtime/Engine/Private/Landscape/Landscape.cpp
+++ b/Engine/Source/Runtime/Engine/Private/Landscape/Landscape.cpp
@@ -1786,8 +1786,8 @@ void ULandscapeInfo::RecreateLandscapeInfo(UWorld* InWorld, bool bMapCheck)
 				FVector	DrawScale = FirstVisibleLandscape->GetRootComponent()->RelativeScale3D;
 				
 				FIntPoint QuadsSpaceOffset; 
-				QuadsSpaceOffset.X = Offset.X/DrawScale.X;
-				QuadsSpaceOffset.Y = Offset.Y/DrawScale.Y;
+				QuadsSpaceOffset.X = FMath::Round(Offset.X/DrawScale.X);
+				QuadsSpaceOffset.Y = FMath::Round(Offset.Y/DrawScale.Y);
 				Proxy->SetAbsoluteSectionBase(QuadsSpaceOffset);
 			}
 
