Fix unity build

diff --git a/Engine/Source/Editor/WorldBrowser/Private/Tiles/STiledLandscapeImportDlg.cpp b/Engine/Source/Editor/WorldBrowser/Private/Tiles/STiledLandscapeImportDlg.cpp
--- a/Engine/Source/Editor/WorldBrowser/Private/Tiles/STiledLandscapeImportDlg.cpp
+++ b/Engine/Source/Editor/WorldBrowser/Private/Tiles/STiledLandscapeImportDlg.cpp
@@ -2,6 +2,9 @@
 
 #include "WorldBrowserPrivatePCH.h"
 #include "STiledLandscapeImportDlg.h"
+#include "SVectorInputBox.h"
+#include "IDesktopPlatform.h"
+#include "DesktopPlatformModule.h"
 
 #define LOCTEXT_NAMESPACE "WorldBrowser"
 
