World Browser: - Tile that has bounds intersect with world bounds will be visible as well

diff --git a/Engine/Source/Editor/WorldBrowser/Private/Tiles/WorldTileModel.cpp b/Engine/Source/Editor/WorldBrowser/Private/Tiles/WorldTileModel.cpp
--- a/Engine/Source/Editor/WorldBrowser/Private/Tiles/WorldTileModel.cpp
+++ b/Engine/Source/Editor/WorldBrowser/Private/Tiles/WorldTileModel.cpp
@@ -242,11 +242,9 @@ bool FWorldTileModel::ShouldBeVisible(FBox EditableArea) const
 	{
 		return true;
 	}
-	
-	// Visible if level bounds bigger than editable area and intersects that area
-	if ((LevelBBox.GetExtent().X >= EditableArea.GetExtent().X || 
-		 LevelBBox.GetExtent().Y >= EditableArea.GetExtent().Y) &&
-		LevelBBox.IntersectXY(EditableArea))
+
+	// Visible if level bounds intersects editable area
+	if (LevelBBox.IntersectXY(EditableArea))
 	{
 		return true;
 	}
