World Browser: - Fixed tiles shelving

diff --git a/Engine/Source/Editor/WorldBrowser/Private/Tiles/WorldTileCollectionModel.cpp b/Engine/Source/Editor/WorldBrowser/Private/Tiles/WorldTileCollectionModel.cpp
--- a/Engine/Source/Editor/WorldBrowser/Private/Tiles/WorldTileCollectionModel.cpp
+++ b/Engine/Source/Editor/WorldBrowser/Private/Tiles/WorldTileCollectionModel.cpp
@@ -1229,7 +1229,7 @@ void FWorldTileCollectionModel::AddLandscapeProxy_Executed(FWorldTileModel::EWor
 		ALandscapeProxy* SourceLandscape = LandscapeTileModel->GetLandcape();
 		FIntPoint SourceTileOffset = LandscapeTileModel->GetAbsoluteLevelPosition();
 
-		NewLevelModel->SetVisible(false);
+		NewLevelModel->Shelve();
 		NewLevelModel->CreateAdjacentLandscapeProxy(SourceLandscape, SourceTileOffset, InWhere);
 		ShowLevels(Levels);
 	}
@@ -1320,7 +1320,7 @@ void FWorldTileCollectionModel::ImportTiledLandscape_Executed()
 					PopulateLevelsList();
 					TSharedPtr<FWorldTileModel> NewTileModel = StaticCastSharedPtr<FWorldTileModel>(FindLevelModel(NewWorld->GetOutermost()->GetFName()));
 					// Hide level, so we do not depend on a current world origin
-					NewTileModel->SetVisible(false);
+					NewTileModel->Shelve();
 					
 					// Create landscape (actor/proxy) in a new level
 					ALandscapeProxy* NewLandscape = NewTileModel->ImportLandscape(TileImportSettings);
diff --git a/Engine/Source/Editor/WorldBrowser/Private/Tiles/WorldTileModel.cpp b/Engine/Source/Editor/WorldBrowser/Private/Tiles/WorldTileModel.cpp
--- a/Engine/Source/Editor/WorldBrowser/Private/Tiles/WorldTileModel.cpp
+++ b/Engine/Source/Editor/WorldBrowser/Private/Tiles/WorldTileModel.cpp
@@ -578,13 +578,19 @@ void FWorldTileModel::LoadLevel()
 	// Create transient level streaming object and add to persistent level
 	ULevelStreaming* LevelStreaming = GetAssosiatedStreamingLevel();
 
+	// Whether this tile should be made visible at current world bounds
+	const bool bShouldBeVisible = ShouldBeVisible(LevelCollectionModel.EditableWorldArea());
+	
 	// Load level
 	bLoadingLevel = true;
 	LevelStreaming->bShouldBeLoaded = true;
-	LevelStreaming->bShouldBeVisible = false;
-	LevelStreaming->bShouldBeVisibleInEditor = ShouldBeVisible(LevelCollectionModel.EditableWorldArea());
+	LevelStreaming->bShouldBeVisible = false; // Should be always false in the Editor
+	LevelStreaming->bShouldBeVisibleInEditor = bShouldBeVisible;
 	LevelCollectionModel.GetWorld()->FlushLevelStreaming();
 	bLoadingLevel = false;
+	// Mark tile as shelved in case it is hidden(does not fit to world bounds)
+	bWasShelved = !bShouldBeVisible;
+	//
 	LoadedLevel = LevelStreaming->GetLoadedLevel();
 	// Enable tile properties
 	TileDetails->bTileEditable = (LoadedLevel != nullptr);
