World Browser: - Changed splitters to be more visible

diff --git a/Engine/Source/Editor/LevelEditor/Private/SLevelEditor.cpp b/Engine/Source/Editor/LevelEditor/Private/SLevelEditor.cpp
--- a/Engine/Source/Editor/LevelEditor/Private/SLevelEditor.cpp
+++ b/Engine/Source/Editor/LevelEditor/Private/SLevelEditor.cpp
@@ -679,12 +679,7 @@ TSharedRef<SDockTab> SLevelEditor::SpawnLevelEditorTab( const FSpawnTabArgs& Arg
 			.Icon( FEditorStyle::GetBrush( "LevelEditor.Tabs.WorldBrowser" ) )
 			.Label( NSLOCTEXT("LevelEditor", "WorldBrowserTabTitle", "World Browser") )
 			[
-				SNew(SBorder)
-				.Padding( 0 )
-				.BorderImage( FEditorStyle::GetBrush("ToolPanel.GroupBorder") )
-				[
-					WorldBrowserModule.CreateWorldBrowser()
-				]
+				WorldBrowserModule.CreateWorldBrowser()
 			];
 	}
 	else if( TabIdentifier == TEXT("Sequencer") && FParse::Param(FCommandLine::Get(), TEXT("sequencer")) )
diff --git a/Engine/Source/Editor/WorldBrowser/Private/SWorldDetailsView.cpp b/Engine/Source/Editor/WorldBrowser/Private/SWorldDetailsView.cpp
--- a/Engine/Source/Editor/WorldBrowser/Private/SWorldDetailsView.cpp
+++ b/Engine/Source/Editor/WorldBrowser/Private/SWorldDetailsView.cpp
@@ -13,7 +13,7 @@ void SWorldDetailsView::Construct(const FArguments& InArgs)
 	
 	FPropertyEditorModule& PropertyModule = FModuleManager::LoadModuleChecked<FPropertyEditorModule>("PropertyEditor");
 
-	FDetailsViewArgs Args(false, false, true, false, true);
+	FDetailsViewArgs Args(false, false, false, false, true);
 	Args.bShowActorLabel = false;
 	
 	DetailsView = PropertyModule.CreateDetailView(Args);
diff --git a/Engine/Source/Editor/WorldBrowser/Private/SWorldMainView.cpp b/Engine/Source/Editor/WorldBrowser/Private/SWorldMainView.cpp
--- a/Engine/Source/Editor/WorldBrowser/Private/SWorldMainView.cpp
+++ b/Engine/Source/Editor/WorldBrowser/Private/SWorldMainView.cpp
@@ -57,7 +57,8 @@ void SWorldMainView::Construct( const FArguments& InArgs )
 		.AutoHeight()
 		[
 			SAssignNew(MainMenuHolder, SBorder)
-			.BorderImage(FEditorStyle::GetBrush("NoBrush"))
+			.Padding(0)
+			.BorderImage(FEditorStyle::GetBrush(TEXT("NoBrush")))
 		]
 
 		// Create parent widget for all world browser views
@@ -65,8 +66,8 @@ void SWorldMainView::Construct( const FArguments& InArgs )
 		.FillHeight(1.0f)
 		[
 			SAssignNew(ContentParent, SBorder)
-			.Padding(5)
-			.BorderImage(FEditorStyle::GetBrush("NoBrush"))
+			.Padding(0)
+			.BorderImage(FEditorStyle::GetBrush(TEXT("NoBrush")))
 		]
 	];
 	
@@ -446,13 +447,18 @@ TSharedRef<SWidget> SWorldMainView::CreateContentViews()
 {
 	TSharedPtr<SSplitter> ViewSplitter;
 	SAssignNew(ViewSplitter, SSplitter);
-	
+
 	// Hierarchy view
 	ViewSplitter->AddSlot()
 	.Value(0.3f)
 	[
-		SNew(SWorldLevelsTreeView)
-		.InWorldModel(WorldModel)
+		SNew(SBorder)
+		.Padding(FMargin(0,3,0,0))
+		.BorderImage(FEditorStyle::GetBrush(TEXT("ToolPanel.GroupBorder")))
+		[
+			SNew(SWorldLevelsTreeView)
+			.InWorldModel(WorldModel)
+		]
 	];
 
 	// Grid view
@@ -461,7 +467,12 @@ TSharedRef<SWidget> SWorldMainView::CreateContentViews()
 		ViewSplitter->AddSlot()
 		.Value(0.7f)
 		[
-			ConstructCenterPane()
+			SNew(SBorder)
+			.Padding(FMargin(0,3,0,0))
+			.BorderImage(FEditorStyle::GetBrush(TEXT("ToolPanel.GroupBorder")))
+			[
+				ConstructCenterPane()
+			]
 		];
 	}
 
@@ -469,9 +480,14 @@ TSharedRef<SWidget> SWorldMainView::CreateContentViews()
 	ViewSplitter->AddSlot()
 	.Value(0.2f)
 	[
-		SNew(SWorldDetailsView)
-		.Visibility(this, &SWorldMainView::GetDetailsViewVisibility)
-		.InWorldModel(WorldModel)
+		SNew(SBorder)
+		.Padding(FMargin(0,0,0,0))
+		.BorderImage(FEditorStyle::GetBrush(TEXT("ToolPanel.GroupBorder")))
+		[
+			SNew(SWorldDetailsView)
+			.Visibility(this, &SWorldMainView::GetDetailsViewVisibility)
+			.InWorldModel(WorldModel)
+		]
 	];
 
 	return ViewSplitter.ToSharedRef();
