World Browser:
- Added icons to level hierarchy to distinguish level streaming methods
- Added tile coordinates offset when umporting tiled heightmap

diff --git a/Engine/Source/Editor/WorldBrowser/Private/LevelCollectionModel.h b/Engine/Source/Editor/WorldBrowser/Private/LevelCollectionModel.h
--- a/Engine/Source/Editor/WorldBrowser/Private/LevelCollectionModel.h
+++ b/Engine/Source/Editor/WorldBrowser/Private/LevelCollectionModel.h
@@ -520,16 +520,17 @@ struct FTiledLandscapeImportSettings
 		, ComponentsNum(8)
 		, QuadsPerSection(63)
 		, SectionsPerComponent(1)
+		, TilesCoordinatesOffset(0,0)
 		, TileResolution(1009)	
 	{}
 	
-	FVector			Scale3D;
-	int32			ComponentsNum;
-	int32			QuadsPerSection;
-	int32			SectionsPerComponent;
+	FVector				Scale3D;
+	int32				ComponentsNum;
+	int32				QuadsPerSection;
+	int32				SectionsPerComponent;
 	
 	TArray<FString>		ImportFileList;
 	TArray<FIntPoint>	TileCoordinates;
-	FIntRect			TotalLandscapeRect;	
+	FIntPoint			TilesCoordinatesOffset;	
 	int32				TileResolution;	
 };
\ No newline at end of file
diff --git a/Engine/Source/Editor/WorldBrowser/Private/LevelModel.cpp b/Engine/Source/Editor/WorldBrowser/Private/LevelModel.cpp
--- a/Engine/Source/Editor/WorldBrowser/Private/LevelModel.cpp
+++ b/Engine/Source/Editor/WorldBrowser/Private/LevelModel.cpp
@@ -716,4 +716,9 @@ FString FLevelModel::GetFileSizeString() const
 	return MemorySizeString;
 }
 
+UClass* FLevelModel::GetStreamingClass() const
+{
+	return nullptr;
+}
+
 #undef LOCTEXT_NAMESPACE
diff --git a/Engine/Source/Editor/WorldBrowser/Private/LevelModel.h b/Engine/Source/Editor/WorldBrowser/Private/LevelModel.h
--- a/Engine/Source/Editor/WorldBrowser/Private/LevelModel.h
+++ b/Engine/Source/Editor/WorldBrowser/Private/LevelModel.h
@@ -245,6 +245,9 @@ public:
 	/** @return the Level's File Size as a FString */
 	FString GetFileSizeString() const;
 	
+	/** @return Class used for streaming this level */
+	virtual UClass* GetStreamingClass() const;
+	
 protected:
 	/** Level model display name */
 	FString								DisplayName;
diff --git a/Engine/Source/Editor/WorldBrowser/Private/SWorldGridItem.cpp b/Engine/Source/Editor/WorldBrowser/Private/SWorldGridItem.cpp
--- a/Engine/Source/Editor/WorldBrowser/Private/SWorldGridItem.cpp
+++ b/Engine/Source/Editor/WorldBrowser/Private/SWorldGridItem.cpp
@@ -135,10 +135,6 @@ void SWorldGridItem::Construct(const FArguments& InArgs)
 	WorldModel = InArgs._InWorldModel;
 	LevelModel = InArgs._InItemModel;
 
-	MapPackageLoaded = FEditorStyle::GetBrush("WorldBrowser.LevelLoaded");
-	MapPackageUnloaded = FEditorStyle::GetBrush("WorldBrowser.LevelUnloaded");
-	MapPackagePending = FEditorStyle::GetBrush("WorldBrowser.LevelPending");
-
 	ProgressBarImage = FEditorStyle::GetBrush(TEXT("ProgressBar.Marquee"));
 	
 	bNeedRefresh = true;
@@ -209,15 +205,6 @@ TSharedRef<SToolTip> SWorldGridItem::CreateToolTipWidget()
 					[
 						SNew(SHorizontalBox)
 
-						// Level status
-						+SHorizontalBox::Slot()
-						.HAlign(HAlign_Left)
-						.AutoWidth()
-						[
-							SNew(SImage)
-							.Image(this, &SWorldGridItem::GetLevelStatusBrush)
-						]
-
 						// Level name
 						+SHorizontalBox::Slot()
 						.Padding(6,0,0,0)
@@ -442,18 +429,6 @@ FString SWorldGridItem::GetLevelNameText() const
 	return LevelModel->GetDisplayName();
 }
 
-const FSlateBrush* SWorldGridItem::GetLevelStatusBrush() const
-{
-	if (LevelModel->IsLoaded())
-	{
-		return MapPackageLoaded;
-	}
-
-	return LevelModel->IsLoading() ? 
-		MapPackagePending : 
-		MapPackageUnloaded;
-}
-
 bool SWorldGridItem::IsItemEditable() const
 {
 	return LevelModel->IsEditable();
diff --git a/Engine/Source/Editor/WorldBrowser/Private/SWorldGridItem.h b/Engine/Source/Editor/WorldBrowser/Private/SWorldGridItem.h
--- a/Engine/Source/Editor/WorldBrowser/Private/SWorldGridItem.h
+++ b/Engine/Source/Editor/WorldBrowser/Private/SWorldGridItem.h
@@ -100,7 +100,6 @@ private:
 	TSharedRef<SToolTip> CreateToolTipWidget();
 	
 	//
-	const FSlateBrush* GetLevelStatusBrush() const;
 	FString GetLevelNameText() const;
 
 public:
@@ -119,10 +118,8 @@ private:
 	TSharedPtr<FGridItemThumbnail>	Thumbnail;
 	/** The actual widget for the thumbnail. */
 	TSharedPtr<SViewport>			ThumbnailViewport;
-
-	const FSlateBrush*				MapPackageUnloaded;
-	const FSlateBrush*				MapPackagePending;
-	const FSlateBrush*				MapPackageLoaded;
+	
+	/** */
 	const FSlateBrush*				ProgressBarImage;
 
 	/** Curve sequence to drive loading animation */
diff --git a/Engine/Source/Editor/WorldBrowser/Private/SWorldTreeItem.cpp b/Engine/Source/Editor/WorldBrowser/Private/SWorldTreeItem.cpp
--- a/Engine/Source/Editor/WorldBrowser/Private/SWorldTreeItem.cpp
+++ b/Engine/Source/Editor/WorldBrowser/Private/SWorldTreeItem.cpp
@@ -16,15 +16,15 @@ void SWorldTreeItem::Construct(const FArguments& InArgs, TSharedRef<STableViewBa
 	IsItemExpanded = InArgs._IsItemExpanded;
 	HighlightText = InArgs._HighlightText;
 
-	MapPackageLoaded = FEditorStyle::GetBrush("WorldBrowser.LevelLoaded");
-	MapPackageUnloaded = FEditorStyle::GetBrush("WorldBrowser.LevelUnloaded");
-	MapPackagePending = FEditorStyle::GetBrush("WorldBrowser.LevelPending");
+	StreamingLevelAlwaysLoadedBrush = FEditorStyle::GetBrush("WorldBrowser.LevelStreamingAlwaysLoaded");
+	StreamingLevelKismetBrush = FEditorStyle::GetBrush("WorldBrowser.LevelStreamingKismet");
 
 	SMultiColumnTableRow<TSharedPtr<FLevelModel>>::Construct(
 		FSuperRowType::FArguments()
 			.OnDragDetected(this, &SWorldTreeItem::OnItemDragDetected), 
 		OwnerTableView
 	);
+
 }
 
 BEGIN_SLATE_FUNCTION_BUILD_OPTIMIZATION
@@ -43,6 +43,18 @@ TSharedRef< SWidget > SWorldTreeItem::GenerateWidgetForColumn( const FName & Col
 				SNew(SExpanderArrow, SharedThis(this))
 			]
 			
+			+SHorizontalBox::Slot()
+			.VAlign(VAlign_Center)
+			.AutoWidth()
+			[
+				SNew(SBox)
+				.WidthOverride(7) // in case brush for this item is empty
+				[
+					SNew(SImage)
+					.Image(this, &SWorldTreeItem::GetLevelIconBrush)
+				]
+			]
+
 			+SHorizontalBox::Slot()
 			.VAlign(VAlign_Center)
 			.AutoWidth()
@@ -229,16 +241,6 @@ FReply SWorldTreeItem::OnOpenKismet()
 	return FReply::Handled();
 }
 
-const FSlateBrush* SWorldTreeItem::GetLevelImage() const
-{
-	if (LevelModel->IsLoaded())
-	{
-		return MapPackageLoaded;
-	}
-
-	return LevelModel->IsLoading() ? MapPackagePending : MapPackageUnloaded;
-}
-
 FReply SWorldTreeItem::OnItemDragDetected(const FGeometry& MyGeometry, const FPointerEvent& MouseEvent)
 {
 	if (MouseEvent.IsMouseButtonDown(EKeys::LeftMouseButton))
@@ -334,6 +336,22 @@ FSlateColor SWorldTreeItem::GetLevelDisplayNameColorAndOpacity() const
 	return FSlateColor::UseForeground();
 }
 
+const FSlateBrush* SWorldTreeItem::GetLevelIconBrush() const
+{
+	UClass* StreamingClass = LevelModel->GetStreamingClass();
+	if (StreamingClass == ULevelStreamingKismet::StaticClass())
+	{
+		return StreamingLevelKismetBrush;
+	}
+
+	if (StreamingClass == ULevelStreamingAlwaysLoaded::StaticClass())
+	{
+		return StreamingLevelAlwaysLoadedBrush;
+	}
+
+	return nullptr;
+}
+
 const FSlateBrush* SWorldTreeItem::GetLevelVisibilityBrush() const
 {
 	if (LevelModel->GetLevelObject())
diff --git a/Engine/Source/Editor/WorldBrowser/Private/SWorldTreeItem.h b/Engine/Source/Editor/WorldBrowser/Private/SWorldTreeItem.h
--- a/Engine/Source/Editor/WorldBrowser/Private/SWorldTreeItem.h
+++ b/Engine/Source/Editor/WorldBrowser/Private/SWorldTreeItem.h
@@ -96,9 +96,9 @@ private:
 	FSlateColor GetLevelDisplayNameColorAndOpacity() const;
 
 	/**
-	 *	@return The brush used to draw the folder icon 
+	 *	@return	The SlateBrush representing level icon
 	 */
-	const FSlateBrush* GetLevelImage() const;
+	const FSlateBrush* GetLevelIconBrush() const;
 
 	/**
 	 *	Called to get the Slate Image Brush representing the visibility state of
@@ -154,11 +154,10 @@ private:
 	/** True when this item has children and is expanded */
 	TAttribute<bool>				IsItemExpanded;
 
-	/** Brushes for the different level states */
-	const FSlateBrush*				MapPackageUnloaded;
-	const FSlateBrush*				MapPackagePending;
-	const FSlateBrush*				MapPackageLoaded;
-	
+	/** Brushes for the different streaming level types */
+	const FSlateBrush*				StreamingLevelAlwaysLoadedBrush;
+	const FSlateBrush*				StreamingLevelKismetBrush;
+
 	/**	The visibility button for the Level */
 	TSharedPtr<SButton>				VisibilityButton;
 
diff --git a/Engine/Source/Editor/WorldBrowser/Private/StreamingLevels/StreamingLevelModel.cpp b/Engine/Source/Editor/WorldBrowser/Private/StreamingLevels/StreamingLevelModel.cpp
--- a/Engine/Source/Editor/WorldBrowser/Private/StreamingLevels/StreamingLevelModel.cpp
+++ b/Engine/Source/Editor/WorldBrowser/Private/StreamingLevels/StreamingLevelModel.cpp
@@ -116,6 +116,16 @@ void FStreamingLevelModel::GetGridItemTooltipFields(TArray< TPair<TAttribute<FTe
 //	return false;
 //}
 
+UClass* FStreamingLevelModel::GetStreamingClass() const
+{
+	if (!IsPersistent() && LevelStreaming.IsValid())
+	{
+		return LevelStreaming.Get()->GetClass();
+	}
+	
+	return FLevelModel::GetStreamingClass();
+}
+
 void FStreamingLevelModel::SetStreamingClass( UClass *LevelStreamingClass )
 {
 	if (IsPersistent())
diff --git a/Engine/Source/Editor/WorldBrowser/Private/StreamingLevels/StreamingLevelModel.h b/Engine/Source/Editor/WorldBrowser/Private/StreamingLevels/StreamingLevelModel.h
--- a/Engine/Source/Editor/WorldBrowser/Private/StreamingLevels/StreamingLevelModel.h
+++ b/Engine/Source/Editor/WorldBrowser/Private/StreamingLevels/StreamingLevelModel.h
@@ -36,6 +36,7 @@ public:
 	virtual void OnDrop(const TSharedPtr<FLevelDragDropOp>& Op) OVERRIDE;
 	virtual bool IsGoodToDrop(const TSharedPtr<FLevelDragDropOp>& Op) const OVERRIDE;
 	virtual void GetGridItemTooltipFields(TArray< TPair<TAttribute<FText>, TAttribute<FText>> >& CustomFields) const OVERRIDE;
+	virtual UClass* GetStreamingClass() const OVERRIDE;
 	// FLevelModel interface end
 		
 	/** @return the StreamingLevelIndex */
@@ -49,7 +50,7 @@ public:
 
 	/** Sets the Level's streaming class */
 	void SetStreamingClass( UClass *LevelStreamingClass );
-		
+	
 private:
 	/** Updates cached value of package file availability */
 	void UpdatePackageFileAvailability();
diff --git a/Engine/Source/Editor/WorldBrowser/Private/Tiles/STiledLandscapeImportDlg.cpp b/Engine/Source/Editor/WorldBrowser/Private/Tiles/STiledLandscapeImportDlg.cpp
--- a/Engine/Source/Editor/WorldBrowser/Private/Tiles/STiledLandscapeImportDlg.cpp
+++ b/Engine/Source/Editor/WorldBrowser/Private/Tiles/STiledLandscapeImportDlg.cpp
@@ -14,12 +14,11 @@ static int32 CalcLandscapeSquareResolution(int32 ComponetsNumX, int32 SectionNum
 }
 
 /**
- *	Returns heigtmap tile coordinates extracted from a specified tile filename
+ *	Returns heightmap tile coordinates extracted from a specified tile filename
  */
 static FIntPoint ExtractHeighmapTileCoordinates(FString BaseFilename)
 {
 	//We expect file name in form: <tilename>_x<number>_y<number>
-		
 	FIntPoint ResultPosition(-1,-1);
 	
 	int32 XPos = BaseFilename.Find(TEXT("_x"), ESearchCase::IgnoreCase, ESearchDir::FromEnd);
@@ -69,16 +68,59 @@ void STiledLandcapeImportDlg::Construct(const FArguments& InArgs, TSharedPtr<SWi
 					.OnClicked(this, &STiledLandcapeImportDlg::OnClickedSelectTiles)
 					.Text(LOCTEXT("TiledLandscapeImport_SelectButtonText", "Select Tiles..."))
 				]
+				
+				// Tiles origin offset
+				+SUniformGridPanel::Slot(0, 1)
+				.VAlign(VAlign_Center)
+				[
+					SNew(STextBlock)
+					.ToolTipText(LOCTEXT("TiledLandscapeImport_TilesOffsetTooltip", "For example: tile x0_y0 will be treated as x(0+offsetX)_y(0+offsetY)"))
+					.Text(LOCTEXT("TiledLandscapeImport_TilesOffsetText", "Tile Coordinates Offset"))
+				]
+
+				+SUniformGridPanel::Slot(1, 1)
+				.VAlign(VAlign_Center)
+				[
+					SNew(SHorizontalBox)
+					
+					+SHorizontalBox::Slot()
+					.Padding(0.0f, 1.0f, 2.0f, 1.0f)
+					.FillWidth(1.f)
+					[
+						SNew(SNumericEntryBox<int32>)
+						.Value(this, &STiledLandcapeImportDlg::GetTileOffsetX)
+						.OnValueChanged(this, &STiledLandcapeImportDlg::SetTileOffsetX)
+						.LabelPadding(0)
+						.Label()
+						[
+							SNumericEntryBox<int>::BuildLabel( LOCTEXT("X_Label", "X"), FLinearColor::White, SNumericEntryBox<int>::RedLabelBackgroundColor )
+						]
+					]
+					
+					+SHorizontalBox::Slot()
+					.Padding(0.0f, 1.0f, 2.0f, 1.0f)
+					.FillWidth(1.f)
+					[
+						SNew(SNumericEntryBox<int32>)
+						.Value(this, &STiledLandcapeImportDlg::GetTileOffsetY)
+						.OnValueChanged(this, &STiledLandcapeImportDlg::SetTileOffsetY)
+						.LabelPadding(0)
+						.Label()
+						[
+							SNumericEntryBox<float>::BuildLabel( LOCTEXT("Y_Label", "Y"), FLinearColor::White, SNumericEntryBox<int>::GreenLabelBackgroundColor )
+						]
+					]
+				]
 
 				// Tile configuration
-				+SUniformGridPanel::Slot(0, 1)
+				+SUniformGridPanel::Slot(0, 2)
 				.VAlign(VAlign_Center)
 				[
 					SNew(STextBlock)
 					.Text(LOCTEXT("TiledLandscapeImport_ConfigurationText", "Import Configuration"))
 				]
 
-				+SUniformGridPanel::Slot(1, 1)
+				+SUniformGridPanel::Slot(1, 2)
 				.VAlign(VAlign_Center)
 				[
 					SAssignNew(TileConfigurationComboBox, SComboBox<TSharedPtr<FTileImportConfiguration>>)
@@ -93,14 +135,14 @@ void STiledLandcapeImportDlg::Construct(const FArguments& InArgs, TSharedPtr<SWi
 				]
 
 				// Scale
-				+SUniformGridPanel::Slot(0, 2)
+				+SUniformGridPanel::Slot(0, 3)
 				.VAlign(VAlign_Center)
 				[
 					SNew(STextBlock)
 					.Text(LOCTEXT("TiledLandscapeImport_ScaleText", "Landscape Scale"))
 				]
 			
-				+SUniformGridPanel::Slot(1, 2)
+				+SUniformGridPanel::Slot(1, 3)
 				.VAlign(VAlign_Center)
 				[
 					SNew( SVectorInputBox )
@@ -112,7 +154,6 @@ void STiledLandcapeImportDlg::Construct(const FArguments& InArgs, TSharedPtr<SWi
 					.OnYCommitted( this, &STiledLandcapeImportDlg::OnSetScale, 1 )
 					.OnZCommitted( this, &STiledLandcapeImportDlg::OnSetScale, 2 )
 				]
-
 			]
 
 			// Import summary
@@ -222,6 +263,26 @@ void STiledLandcapeImportDlg::OnSetScale(float InValue, ETextCommit::Type Commit
 	}
 }
 
+TOptional<int32> STiledLandcapeImportDlg::GetTileOffsetX() const
+{
+	return ImportSettings.TilesCoordinatesOffset.X;
+}
+
+void STiledLandcapeImportDlg::SetTileOffsetX(int32 InValue)
+{
+	ImportSettings.TilesCoordinatesOffset.X = InValue;
+}
+
+TOptional<int32> STiledLandcapeImportDlg::GetTileOffsetY() const
+{
+	return ImportSettings.TilesCoordinatesOffset.Y;
+}
+
+void STiledLandcapeImportDlg::SetTileOffsetY(int32 InValue)
+{
+	ImportSettings.TilesCoordinatesOffset.Y = InValue;
+}
+
 void STiledLandcapeImportDlg::OnSetImportConfiguration(TSharedPtr<FTileImportConfiguration> InTileConfig, ESelectInfo::Type SelectInfo)
 {
 	if (InTileConfig.IsValid())
@@ -239,10 +300,10 @@ void STiledLandcapeImportDlg::OnSetImportConfiguration(TSharedPtr<FTileImportCon
 
 FReply STiledLandcapeImportDlg::OnClickedSelectTiles()
 {
+	TotalLandscapeRect = FIntRect(MAX_int32, MAX_int32, MIN_int32, MIN_int32);
 	ImportSettings.ImportFileList.Empty();
 	ImportSettings.TileCoordinates.Empty();
-	ImportSettings.TotalLandscapeRect = FIntRect(MAX_int32, MAX_int32, MIN_int32, MIN_int32);
-
+	
 	SetPossibleConfigurationsForResolution(0);
 	
 	IDesktopPlatform* DesktopPlatform = FDesktopPlatformModule::Get();
@@ -284,8 +345,9 @@ FReply STiledLandcapeImportDlg::OnClickedSelectTiles()
 						break;
 					}
 
+					TotalLandscapeRect.Include(TileCoordinate);
 					ImportSettings.TileCoordinates.Add(TileCoordinate);
-					ImportSettings.TotalLandscapeRect.Include(TileCoordinate);
+					
 				}
 				
 				if (bValidTiles)
@@ -385,8 +447,8 @@ FText STiledLandcapeImportDlg::GetImportSummaryText() const
 	const FString TilesSummary = FString::Printf(TEXT("%d - %dx%d"), ImportSettings.ImportFileList.Num(), ImportSettings.TileResolution, ImportSettings.TileResolution);
 	
 	// Total landscape size(NxN km)
-	int32 WidthInTilesX = ImportSettings.TotalLandscapeRect.Width() + 1;
-	int32 WidthInTilesY = ImportSettings.TotalLandscapeRect.Height() + 1;
+	int32 WidthInTilesX = TotalLandscapeRect.Width() + 1;
+	int32 WidthInTilesY = TotalLandscapeRect.Height() + 1;
 	float WidthX = 0.00001f*ImportSettings.Scale3D.X*WidthInTilesX*ImportSettings.TileResolution;
 	float WidthY = 0.00001f*ImportSettings.Scale3D.Y*WidthInTilesY*ImportSettings.TileResolution;
 	const FString LandscapeSummary = FString::Printf(TEXT("%.3fx%.3f"), WidthX, WidthY);
diff --git a/Engine/Source/Editor/WorldBrowser/Private/Tiles/STiledLandscapeImportDlg.h b/Engine/Source/Editor/WorldBrowser/Private/Tiles/STiledLandscapeImportDlg.h
--- a/Engine/Source/Editor/WorldBrowser/Private/Tiles/STiledLandscapeImportDlg.h
+++ b/Engine/Source/Editor/WorldBrowser/Private/Tiles/STiledLandscapeImportDlg.h
@@ -18,7 +18,15 @@ public:
 		int32 NumQuadsPerSection;
 	};
 
-
+	struct FExistingLandscape
+	{
+		FGuid	LandcapeGuid;	
+		FString LandcapeLabel;
+		int32	NumComponents;
+		int32	NumSectionsPerComponent;
+		int32	NumQuadsPerSection;
+	};
+	
 	SLATE_BEGIN_ARGS( STiledLandcapeImportDlg )
 		{}
 	SLATE_END_ARGS()
@@ -38,6 +46,12 @@ private:
 	TOptional<float> GetScaleZ() const;
 	void OnSetScale(float InValue, ETextCommit::Type CommitType, int32 InAxis);
 
+	/** */
+	TOptional<int32> GetTileOffsetX() const;
+	TOptional<int32> GetTileOffsetY() const;
+	void SetTileOffsetX(int32 InValue);
+	void SetTileOffsetY(int32 InValue);
+	
 	/** */
 	void OnSetImportConfiguration(TSharedPtr<FTileImportConfiguration> InTileConfig, ESelectInfo::Type SelectInfo);
 
@@ -69,7 +83,13 @@ private:
 	TSharedPtr<SComboBox<TSharedPtr<FTileImportConfiguration>>> TileConfigurationComboBox;
 	
 	/** */
-	FTiledLandscapeImportSettings	ImportSettings;
+	TSharedPtr<SComboBox<TSharedPtr<FExistingLandscape>>> ExistingLandscapeActors;
+	
+	/** */
+	FTiledLandscapeImportSettings ImportSettings;
+	
+	/** */
+	FIntRect TotalLandscapeRect;	
 	
 	TArray<FTileImportConfiguration> AllConfigurations;
 	TArray<TSharedPtr<FTileImportConfiguration>> ActiveConfigurations;
diff --git a/Engine/Source/Editor/WorldBrowser/Private/Tiles/WorldTileCollectionModel.cpp b/Engine/Source/Editor/WorldBrowser/Private/Tiles/WorldTileCollectionModel.cpp
--- a/Engine/Source/Editor/WorldBrowser/Private/Tiles/WorldTileCollectionModel.cpp
+++ b/Engine/Source/Editor/WorldBrowser/Private/Tiles/WorldTileCollectionModel.cpp
@@ -1344,16 +1344,7 @@ void FWorldTileCollectionModel::ImportTiledLandscape_Executed()
 						NewTileModel->TileDetails->Bounds = NewLandscape->GetComponentsBoundingBox();
 
 						// Calculate this tile offset from world origin
-						FIntPoint TileCoordinates = ImportSettings.TileCoordinates[TileIndex];
-						
-						// Mirror Y coordinate, to match unreal coordinate system
-						// TileCoordinates.Y = TilesRect.Height() - TileCoordinates.Y; //option
-
-						// Tiles coordinates are in range [0..N], convert them to [-N/2..N/2] so whole landscape will be centered around world origin
-						int32 WidthInTilesX = ImportSettings.TotalLandscapeRect.Width() + 1;
-						int32 WidthInTilesY = ImportSettings.TotalLandscapeRect.Height() + 1;
-						TileCoordinates = TileCoordinates - FIntPoint(WidthInTilesX/2, WidthInTilesY/2);
-												
+						FIntPoint TileCoordinates = ImportSettings.TileCoordinates[TileIndex] + ImportSettings.TilesCoordinatesOffset;
 						FIntPoint TileOffset = FIntPoint(TileCoordinates.X*WidthX, TileCoordinates.Y*WidthY);
 						
 						// Place level tile at correct position in the world
