World browser:
- Fixed [missing level] operations: replace, remove
- Restored correct styles from CL# 2047641

diff --git a/Engine/Source/Editor/EditorStyle/Private/SlateEditorStyle.cpp b/Engine/Source/Editor/EditorStyle/Private/SlateEditorStyle.cpp
--- a/Engine/Source/Editor/EditorStyle/Private/SlateEditorStyle.cpp
+++ b/Engine/Source/Editor/EditorStyle/Private/SlateEditorStyle.cpp
@@ -1032,26 +1032,7 @@ void FSlateEditorStyle::FStyle::SetupGeneralStyles()
 
 	// World Browser
 	{
-		Set( "WorldBrowser.Icon16x", new IMAGE_BRUSH( "Icons/levels_16x", Icon16x16 ) );
-		Set( "WorldBrowser.VisibleIcon16x", new IMAGE_BRUSH( "Icons/icon_layer_visible", Icon16x16 ) );
-		Set( "WorldBrowser.NotVisibleIcon16x", new IMAGE_BRUSH( "Icons/icon_layer_not_visible", Icon16x16 ) );
-		Set( "WorldBrowser.LockedIcon16x", new IMAGE_BRUSH( "Icons/padlock_locked_16x", Icon16x16 ) );
-		Set( "WorldBrowser.UnlockedIcon16x", new IMAGE_BRUSH( "Icons/padlock_unlocked_16x", Icon16x16 ) );
-		Set( "WorldBrowser.SaveIcon16x", new IMAGE_BRUSH( "Icons/Save_16x", Icon16x16 ) );
-		Set( "WorldBrowser.SaveModifiedIcon16x", new IMAGE_BRUSH( "Icons/SaveModified_16x", Icon16x16 ) );
-		Set( "WorldBrowser.SaveDisabledIcon16x", new IMAGE_BRUSH( "Icons/save_disabled_16x", Icon16x16 ) );
-		Set( "WorldBrowser.ScriptIcon16x", new IMAGE_BRUSH( "Icons/ActorIcons/PlayerStart_16x", Icon16x16 ) );
-		Set( "WorldBrowser.EmptyIcon16x", new IMAGE_BRUSH( "Icons/Empty_16x", Icon16x16 ) );
-		Set( "WorldBrowser.ColorIcon40x", new IMAGE_BRUSH( "Icons/icon_stop_40x", Icon16x16 ) );
-		Set( "WorldBrowser.LevelLoaded", new IMAGE_BRUSH( "Common/TaskGraph_Normal", Icon16x16 ) );
-		Set( "WorldBrowser.LevelUnloaded", new IMAGE_BRUSH( "Common/TaskGraph_Mono", Icon16x16 ) );
-		Set( "WorldBrowser.LevelPending", new IMAGE_BRUSH( "Common/TaskGraph_Selected", Icon16x16 ) );
 		Set( "WorldBrowser.AddLayer", new IMAGE_BRUSH( "Icons/icon_add_40x", Icon16x16 ) );
-		Set( "WorldBrowser.RegularTreeItemFont", TTF_CORE_FONT( "Fonts/Roboto-Regular", 10 ) );
-		Set( "WorldBrowser.BoldTreeItemFont", TTF_CORE_FONT( "Fonts/Roboto-Bold", 10 ) );
-		Set( "WorldBrowser.StreamingAlwaysLoaded", new IMAGE_BRUSH( "Icons/icon_simulate_40x", Icon16x16 ) );
-		Set( "WorldBrowser.StreamingKismet", new IMAGE_BRUSH( "Icons/icon_kismet2_40x", Icon16x16 ) );
-		Set( "WorldBrowser.StreamingLevelBounds", new IMAGE_BRUSH( "Icons/icon_box_40x", Icon16x16 ) );
 		Set( "WorldBrowser.SimulationViewPositon", new IMAGE_BRUSH( "Icons/icon_variable_used_16x", Icon16x16 ) );
 		Set( "WorldBrowser.MouseLocation", new IMAGE_BRUSH( "Icons/cursor_cardinal_cross", Icon16x16 ) );
 		Set( "WorldBrowser.MarqueeRectSize", new IMAGE_BRUSH( "Icons/icon_ViewMode_VisualisationGBuffer_16px", Icon16x16 ) );
@@ -1062,6 +1043,8 @@ void FSlateEditorStyle::FStyle::SetupGeneralStyles()
 		Set( "WorldBrowser.DirectionXNegative", new IMAGE_BRUSH( "Icons/icon_PanLeft", Icon16x16 ) );
 		Set( "WorldBrowser.DirectionYPositive", new IMAGE_BRUSH( "Icons/icon_PanUp", Icon16x16 ) );
 		Set( "WorldBrowser.DirectionYNegative", new IMAGE_BRUSH( "Icons/icon_PanDown", Icon16x16 ) );
+		Set( "WorldBrowser.LevelStreamingAlwaysLoaded", new FSlateNoResource() );
+		Set( "WorldBrowser.LevelStreamingKismet", new IMAGE_BRUSH( "Icons/ActorIcons/ComponentMobilityStationary_7x16px", Icon7x16 ) );
 
 		Set( "WorldBrowser.StatusBarText", FTextBlockStyle(NormalText)
 				.SetFont( TTF_FONT( "Fonts/Roboto-BoldCondensed", 12 ) )
diff --git a/Engine/Source/Editor/Levels/Private/LevelCollectionViewCommands.h b/Engine/Source/Editor/Levels/Private/LevelCollectionViewCommands.h
--- a/Engine/Source/Editor/Levels/Private/LevelCollectionViewCommands.h
+++ b/Engine/Source/Editor/Levels/Private/LevelCollectionViewCommands.h
@@ -28,8 +28,8 @@ PRAGMA_DISABLE_OPTIMIZATION
 		UI_COMMAND( MoveActorsToSelected, "Move Selected Actors to Level", "Moves the selected actors to this level", EUserInterfaceActionType::Button, FInputGesture() );
 
 		//invalid selected level
-		UI_COMMAND( FixUpInvalidReference, "Replace Selected Level","Removes the broken level and replaces it with the level browsed to", EUserInterfaceActionType::Button, FInputGesture() );
-		UI_COMMAND( RemoveInvalidReference, "Remove Selected Level", "Removes the reference to the missing level from the map", EUserInterfaceActionType::Button, FInputGesture() );
+		UI_COMMAND( FixUpInvalidReference, "Replace Missing Level","Removes the broken level and replaces it with the level browsed to", EUserInterfaceActionType::Button, FInputGesture() );
+		UI_COMMAND( RemoveInvalidReference, "Remove Missing Level", "Removes the reference to the missing level from the map", EUserInterfaceActionType::Button, FInputGesture() );
 
 		//levels
 		UI_COMMAND( EditProperties, "Edit Properties", "Opens the edit properties window for the selected levels", EUserInterfaceActionType::Button, FInputGesture() );
diff --git a/Engine/Source/Editor/WorldBrowser/Private/LevelCollectionModel.cpp b/Engine/Source/Editor/WorldBrowser/Private/LevelCollectionModel.cpp
--- a/Engine/Source/Editor/WorldBrowser/Private/LevelCollectionModel.cpp
+++ b/Engine/Source/Editor/WorldBrowser/Private/LevelCollectionModel.cpp
@@ -664,7 +664,7 @@ void FLevelCollectionModel::BuildHierarchyMenu(FMenuBuilder& InMenuBuilder) cons
 		InMenuBuilder.AddMenuEntry( Commands.AddsActors );
 		InMenuBuilder.AddMenuEntry( Commands.RemovesActors );
 
-		if (AreAnyLevelsSelected())
+		if (AreAnyLevelsSelected() && !(IsOneLevelSelected() && SelectedLevelsList[0]->IsPersistent()))
 		{
 			InMenuBuilder.AddMenuEntry( Commands.SelectStreamingVolumes );
 		}
diff --git a/Engine/Source/Editor/WorldBrowser/Private/LevelModel.cpp b/Engine/Source/Editor/WorldBrowser/Private/LevelModel.cpp
--- a/Engine/Source/Editor/WorldBrowser/Private/LevelModel.cpp
+++ b/Engine/Source/Editor/WorldBrowser/Private/LevelModel.cpp
@@ -652,21 +652,21 @@ void FLevelModel::UpdateLevelActorsCount()
 
 void FLevelModel::UpdateDisplayName()
 {
-	if (HasValidPackage())
+	if (IsPersistent())
 	{
-		if (IsPersistent())
+		DisplayName = LOCTEXT("PersistentTag", "Persistent Level").ToString();
+	}
+	else
+	{
+		DisplayName = GetLongPackageName().ToString();
+		if (!LevelCollectionModel.GetDisplayPathsState())
 		{
-			DisplayName = LOCTEXT("PersistentTag", "Persistent Level").ToString();
-		}
-		else
-		{
-			DisplayName = GetLongPackageName().ToString();
-			if (!LevelCollectionModel.GetDisplayPathsState())
-			{
-				DisplayName = FPackageName::GetShortName(DisplayName);
-			}
+			DisplayName = FPackageName::GetShortName(DisplayName);
 		}
+	}
 
+	if (HasValidPackage())
+	{
 		// Append actors count
 		if (IsLoaded())
 		{
@@ -677,11 +677,10 @@ void FLevelModel::UpdateDisplayName()
 	}
 	else
 	{
-		DisplayName = LOCTEXT("MissingLevelErrorText", " [Missing Level]").ToString();
+		DisplayName+= LOCTEXT("MissingLevelErrorText", " [Missing Level] ").ToString();
 	}
 }
 
-
 FString FLevelModel::GetLightmassSizeString() const
 {
 	FString MemorySizeString;
diff --git a/Engine/Source/Editor/WorldBrowser/Private/StreamingLevels/StreamingLevelCollectionModel.cpp b/Engine/Source/Editor/WorldBrowser/Private/StreamingLevels/StreamingLevelCollectionModel.cpp
--- a/Engine/Source/Editor/WorldBrowser/Private/StreamingLevels/StreamingLevelCollectionModel.cpp
+++ b/Engine/Source/Editor/WorldBrowser/Private/StreamingLevels/StreamingLevelCollectionModel.cpp
@@ -67,6 +67,19 @@ void FStreamingLevelCollectionModel::OnLevelsCollectionChanged()
 	FLevelCollectionModel::OnLevelsCollectionChanged();
 }
 
+void FStreamingLevelCollectionModel::OnLevelsSelectionChanged()
+{
+	InvalidSelectedLevels.Empty();
+
+	for (TSharedPtr<FLevelModel> LevelModel : SelectedLevelsList)
+	{
+		if (!LevelModel->HasValidPackage())
+		{
+			InvalidSelectedLevels.Add(LevelModel);
+		}
+	}
+}
+
 void FStreamingLevelCollectionModel::UnloadLevels(const FLevelModelList& InLevelList)
 {
 	if (IsReadOnly())
@@ -203,31 +216,27 @@ void FStreamingLevelCollectionModel::BuildHierarchyMenu(FMenuBuilder& InMenuBuil
 {
 	const FLevelCollectionCommands& Commands = FLevelCollectionCommands::Get();
 
-	const FLevelModelList& SelectedLevels = GetSelectedLevels();
-			
-	bool bOnlyInvalidSelectedLevel	= InvalidSelectedLevels.Num() == 1 && SelectedLevels.Num() == 0;
-	bool bOnlySelectedLevel			= InvalidSelectedLevels.Num() == 0 && SelectedLevels.Num() == 1;
-			
-	if (bOnlyInvalidSelectedLevel)
+	if (IsOneLevelSelected())
 	{
 		// We only show the level missing level commands if only 1 invalid level and no valid levels
-		InMenuBuilder.BeginSection("MissingLevel", LOCTEXT("ViewHeaderRemove", "Missing Level") );
+		if (InvalidSelectedLevels.Num() == 1)
 		{
-			InMenuBuilder.AddMenuEntry( Commands.FixUpInvalidReference );
-			InMenuBuilder.AddMenuEntry( Commands.RemoveInvalidReference );
+			InMenuBuilder.BeginSection("MissingLevel", LOCTEXT("ViewHeaderRemove", "Missing Level") );
+			{
+				InMenuBuilder.AddMenuEntry( Commands.FixUpInvalidReference );
+				InMenuBuilder.AddMenuEntry( Commands.RemoveInvalidReference );
+			}
+			InMenuBuilder.EndSection();
 		}
-		InMenuBuilder.EndSection();
-	}
-	else if ( bOnlySelectedLevel )
-	{
-		const auto& LevelModel = SelectedLevels[0];
-
-		InMenuBuilder.BeginSection("Level", FText::FromName(LevelModel->GetLongPackageName()));
+		else
 		{
-			InMenuBuilder.AddMenuEntry( Commands.World_MakeLevelCurrent );
-			InMenuBuilder.AddMenuEntry( Commands.MoveActorsToSelected );
+			InMenuBuilder.BeginSection("Level", FText::FromName(GetSelectedLevels()[0]->GetLongPackageName()));
+			{
+				InMenuBuilder.AddMenuEntry( Commands.World_MakeLevelCurrent );
+				InMenuBuilder.AddMenuEntry( Commands.MoveActorsToSelected );
+			}
+			InMenuBuilder.EndSection();
 		}
-		InMenuBuilder.EndSection();
 	}
 
 	// Adds common commands
@@ -235,7 +244,7 @@ void FStreamingLevelCollectionModel::BuildHierarchyMenu(FMenuBuilder& InMenuBuil
 
 	InMenuBuilder.BeginSection("LevelsAddChangeStreaming");
 	{
-		if (AreAnyLevelsSelected())
+		if (AreAnyLevelsSelected() && !(IsOneLevelSelected() && GetSelectedLevels()[0]->IsPersistent()))
 		{
 			InMenuBuilder.AddMenuEntry(Commands.World_RemoveSelectedLevels);
 			
@@ -524,18 +533,22 @@ void FStreamingLevelCollectionModel::AddSelectedActorsToNewLevel_Executed()
 
 void FStreamingLevelCollectionModel::FixupInvalidReference_Executed()
 {
+	// Save or selected list, adding a new level will clean it up
+	FLevelModelList SavedInvalidSelectedLevels = InvalidSelectedLevels;
+
 	// Browsing is essentially the same as adding an existing level
 	if (AddExistingLevel())
 	{
+		InvalidSelectedLevels = SavedInvalidSelectedLevels;
 		RemoveInvalidSelectedLevels_Executed();
 	}
 }
 
 void FStreamingLevelCollectionModel::RemoveInvalidSelectedLevels_Executed()
 {
-	for (auto It = InvalidSelectedLevels.CreateIterator(); It; ++It )
+	for (TSharedPtr<FLevelModel> LevelModel : InvalidSelectedLevels)
 	{
-		TSharedPtr<FStreamingLevelModel> TargetModel = StaticCastSharedPtr<FStreamingLevelModel>(*It);
+		TSharedPtr<FStreamingLevelModel> TargetModel = StaticCastSharedPtr<FStreamingLevelModel>(LevelModel);
 		ULevelStreaming* LevelStreaming = TargetModel->GetLevelStreaming().Get();
 
 		if (LevelStreaming)
diff --git a/Engine/Source/Editor/WorldBrowser/Private/StreamingLevels/StreamingLevelCollectionModel.h b/Engine/Source/Editor/WorldBrowser/Private/StreamingLevels/StreamingLevelCollectionModel.h
--- a/Engine/Source/Editor/WorldBrowser/Private/StreamingLevels/StreamingLevelCollectionModel.h
+++ b/Engine/Source/Editor/WorldBrowser/Private/StreamingLevels/StreamingLevelCollectionModel.h
@@ -42,6 +42,7 @@ private:
 	virtual void Initialize() OVERRIDE;
 	virtual void BindCommands() OVERRIDE;
 	virtual void OnLevelsCollectionChanged() OVERRIDE;
+	virtual void OnLevelsSelectionChanged() OVERRIDE;
 	/** FLevelCollection interface end */
 	
 public:
