- Fixed selection in World Browser

diff --git a/Engine/Source/Editor/WorldBrowser/Private/StreamingLevels/StreamingLevelCollectionModel.cpp b/Engine/Source/Editor/WorldBrowser/Private/StreamingLevels/StreamingLevelCollectionModel.cpp
--- a/Engine/Source/Editor/WorldBrowser/Private/StreamingLevels/StreamingLevelCollectionModel.cpp
+++ b/Engine/Source/Editor/WorldBrowser/Private/StreamingLevels/StreamingLevelCollectionModel.cpp
@@ -78,6 +78,8 @@ void FStreamingLevelCollectionModel::OnLevelsSelectionChanged()
 			InvalidSelectedLevels.Add(LevelModel);
 		}
 	}
+
+	FLevelCollectionModel::OnLevelsSelectionChanged();
 }
 
 void FStreamingLevelCollectionModel::UnloadLevels(const FLevelModelList& InLevelList)
