- Changed World Browser UI layout: hierarchy and detail panel splitted horizontally

diff --git a/Engine/Source/Editor/WorldBrowser/Private/SWorldMainView.cpp b/Engine/Source/Editor/WorldBrowser/Private/SWorldMainView.cpp
--- a/Engine/Source/Editor/WorldBrowser/Private/SWorldMainView.cpp
+++ b/Engine/Source/Editor/WorldBrowser/Private/SWorldMainView.cpp
@@ -447,16 +447,38 @@ TSharedRef<SWidget> SWorldMainView::CreateContentViews()
 	TSharedPtr<SSplitter> ViewSplitter;
 	SAssignNew(ViewSplitter, SSplitter);
 
-	// Hierarchy view
+	// Left side
 	ViewSplitter->AddSlot()
 	.Value(0.3f)
 	[
-		SNew(SBorder)
-		.Padding(FMargin(0,3,0,0))
-		.BorderImage(FEditorStyle::GetBrush(TEXT("ToolPanel.GroupBorder")))
+		SNew(SSplitter)
+		.Orientation(Orient_Vertical)
+
+		// Hierarchy view
+		+SSplitter::Slot()
+		.Value(0.7f)
 		[
-			SNew(SWorldLevelsTreeView)
-			.InWorldModel(WorldModel)
+			SNew(SBorder)
+			.Padding(FMargin(0,3,0,0))
+			.BorderImage(FEditorStyle::GetBrush(TEXT("ToolPanel.GroupBorder")))
+			[
+				SNew(SWorldLevelsTreeView)
+				.InWorldModel(WorldModel)
+			]
+		]
+			
+		// Details view
+		+SSplitter::Slot()
+		.Value(0.3f)
+		[
+			SNew(SBorder)
+			.Padding(FMargin(0,0,0,0))
+			.BorderImage(FEditorStyle::GetBrush(TEXT("ToolPanel.GroupBorder")))
+			.Visibility(this, &SWorldMainView::GetDetailsViewVisibility)
+			[
+				SNew(SWorldDetailsView)
+				.InWorldModel(WorldModel)
+			]
 		]
 	];
 
@@ -475,21 +497,7 @@ TSharedRef<SWidget> SWorldMainView::CreateContentViews()
 			]
 		];
 	}
-
-	// Details view
-	ViewSplitter->AddSlot()
-	.Value(0.2f)
-	[
-		SNew(SBorder)
-		.Padding(FMargin(0,0,0,0))
-		.BorderImage(FEditorStyle::GetBrush(TEXT("ToolPanel.GroupBorder")))
-		.Visibility(this, &SWorldMainView::GetDetailsViewVisibility)
-		[
-			SNew(SWorldDetailsView)
-			.InWorldModel(WorldModel)
-		]
-	];
-
+	
 	return ViewSplitter.ToSharedRef();
 }
 
