World Browser:
- Added material selection to tiled landscape importing tool
- Added tiled weightmap import option for each landscape layer found in selected material

diff --git a/Engine/Source/Editor/WorldBrowser/Private/LevelCollectionModel.h b/Engine/Source/Editor/WorldBrowser/Private/LevelCollectionModel.h
--- a/Engine/Source/Editor/WorldBrowser/Private/LevelCollectionModel.h
+++ b/Engine/Source/Editor/WorldBrowser/Private/LevelCollectionModel.h
@@ -521,16 +521,25 @@ struct FTiledLandscapeImportSettings
 		, QuadsPerSection(63)
 		, SectionsPerComponent(1)
 		, TilesCoordinatesOffset(0,0)
-		, TileResolution(1009)	
+		, TileResolution(1009)
+		, LandscapeMaterial(nullptr)
 	{}
 	
 	FVector				Scale3D;
 	int32				ComponentsNum;
 	int32				QuadsPerSection;
 	int32				SectionsPerComponent;
-	
-	TArray<FString>		ImportFileList;
+
+	TArray<FString>		HeightmapFileList;
 	TArray<FIntPoint>	TileCoordinates;
 	FIntPoint			TilesCoordinatesOffset;	
 	int32				TileResolution;	
+
+
+	UMaterialInterface*					LandscapeMaterial;
+	
+	TArray<FName>						LandscapeLayerNameList;
+	// list of weightmap files per each layer
+	TArray<TArray<FString>>				WeightmapFileList;
+
 };
\ No newline at end of file
diff --git a/Engine/Source/Editor/WorldBrowser/Private/Tiles/STiledLandscapeImportDlg.cpp b/Engine/Source/Editor/WorldBrowser/Private/Tiles/STiledLandscapeImportDlg.cpp
--- a/Engine/Source/Editor/WorldBrowser/Private/Tiles/STiledLandscapeImportDlg.cpp
+++ b/Engine/Source/Editor/WorldBrowser/Private/Tiles/STiledLandscapeImportDlg.cpp
@@ -38,6 +38,11 @@ static FIntPoint ExtractHeighmapTileCoordinates(FString BaseFilename)
 	return ResultPosition;
 }
 
+void STiledLandcapeImportDlg::AddReferencedObjects( FReferenceCollector& Collector )
+{
+	Collector.AddReferencedObject(ImportSettings.LandscapeMaterial);
+}
+
 void STiledLandcapeImportDlg::Construct(const FArguments& InArgs, TSharedPtr<SWindow> InParentWindow)
 {
 	ParentWindow = InParentWindow;
@@ -65,8 +70,8 @@ void STiledLandcapeImportDlg::Construct(const FArguments& InArgs, TSharedPtr<SWi
 					SNew(SButton)
 					.HAlign(HAlign_Center)
 					.ContentPadding(FEditorStyle::GetMargin("StandardDialog.ContentPadding"))
-					.OnClicked(this, &STiledLandcapeImportDlg::OnClickedSelectTiles)
-					.Text(LOCTEXT("TiledLandscapeImport_SelectButtonText", "Select Tiles..."))
+					.OnClicked(this, &STiledLandcapeImportDlg::OnClickedSelectHeightmapTiles)
+					.Text(LOCTEXT("TiledLandscapeImport_SelectButtonText", "Select Heightmap Tiles..."))
 				]
 				
 				// Tiles origin offset
@@ -154,6 +159,42 @@ void STiledLandcapeImportDlg::Construct(const FArguments& InArgs, TSharedPtr<SWi
 					.OnYCommitted( this, &STiledLandcapeImportDlg::OnSetScale, 1 )
 					.OnZCommitted( this, &STiledLandcapeImportDlg::OnSetScale, 2 )
 				]
+
+				// Landcape material
+				+SUniformGridPanel::Slot(0, 4)
+				.VAlign(VAlign_Center)
+				[
+					SNew(STextBlock)
+					.Text(LOCTEXT("TiledLandscapeImport_MaterialText", "Material"))
+				]
+			
+				+SUniformGridPanel::Slot(1, 4)
+				.VAlign(VAlign_Center)
+				[
+					SAssignNew(LandscapeMaterialComboButton, SComboButton)
+					.VAlign(EVerticalAlignment::VAlign_Center)
+					.ComboButtonStyle(FEditorStyle::Get(), "ToolbarComboButton")
+					.ForegroundColor(FLinearColor::White)
+					.ContentPadding(3)
+					.MenuPlacement(EMenuPlacement::MenuPlacement_BelowAnchor)
+					.OnGetMenuContent(this, &STiledLandcapeImportDlg::CreateLandscapeMaterialPicker)
+					.ButtonContent()
+					[
+						SNew( STextBlock )
+						.Text(this, &STiledLandcapeImportDlg::GetLandscapeMaterialName)
+					]
+				]
+			]
+
+			// Layers
+			+SVerticalBox::Slot()
+			.AutoHeight()
+			.Padding(FEditorStyle::GetMargin("StandardDialog.ContentPadding"))
+			[
+				SAssignNew(LayerDataListView, SListView<TSharedPtr<FLandscapeImportLayerData>>)
+				.ListItemsSource( &LayerDataList )
+				.OnGenerateRow( this, &STiledLandcapeImportDlg::OnGenerateWidgetForLayerDataListView )
+				.SelectionMode(ESelectionMode::None)
 			]
 
 			// Import summary
@@ -204,6 +245,36 @@ void STiledLandcapeImportDlg::Construct(const FArguments& InArgs, TSharedPtr<SWi
 	SetPossibleConfigurationsForResolution(0);
 }
 
+TSharedRef<SWidget> STiledLandcapeImportDlg::CreateLandscapeMaterialPicker()
+{
+	FContentBrowserModule& ContentBrowserModule = FModuleManager::Get().LoadModuleChecked<FContentBrowserModule>(TEXT("ContentBrowser"));
+
+	FAssetPickerConfig AssetPickerConfig;
+	AssetPickerConfig.Filter.ClassNames.Add(UMaterial::StaticClass()->GetFName());
+	AssetPickerConfig.Filter.ClassNames.Add(UMaterialInstance::StaticClass()->GetFName());
+	AssetPickerConfig.OnAssetSelected = FOnAssetSelected::CreateSP(this, &STiledLandcapeImportDlg::OnLandscapeMaterialChanged);
+	AssetPickerConfig.InitialAssetViewType = EAssetViewType::List;
+	AssetPickerConfig.bAllowNullSelection = true;
+	AssetPickerConfig.ThumbnailScale = 0.0f;
+	
+	return SNew(SBox)
+		.WidthOverride(250)
+		.HeightOverride(300)
+		[
+			ContentBrowserModule.Get().CreateAssetPicker(AssetPickerConfig)
+		];
+}
+
+FText STiledLandcapeImportDlg::GetLandscapeMaterialName() const
+{
+	if (ImportSettings.LandscapeMaterial)
+	{
+		return FText::FromString(ImportSettings.LandscapeMaterial->GetName());
+	}
+
+	return LOCTEXT("TiledLandscapeImport_NoLandscapeMaterialText", "Default Material");
+}
+
 TSharedRef<SWidget> STiledLandcapeImportDlg::HandleTileConfigurationComboBoxGenarateWidget(TSharedPtr<FTileImportConfiguration> InItem) const
 {
 	const FText ItemText = GenerateConfigurationText(InItem->NumComponents, InItem->NumSectionsPerComponent, InItem->NumQuadsPerSection);
@@ -217,7 +288,7 @@ TSharedRef<SWidget> STiledLandcapeImportDlg::HandleTileConfigurationComboBoxGena
 
 FText STiledLandcapeImportDlg::GetTileConfigurationText() const
 {
-	if (ImportSettings.ImportFileList.Num() == 0)
+	if (ImportSettings.HeightmapFileList.Num() == 0)
 	{
 		return LOCTEXT("TiledLandscapeImport_NoTilesText", "No tiles selected");
 	}
@@ -230,6 +301,47 @@ FText STiledLandcapeImportDlg::GetTileConfigurationText() const
 	return GenerateConfigurationText(ImportSettings.ComponentsNum, ImportSettings.SectionsPerComponent, ImportSettings.QuadsPerSection);
 }
 
+TSharedRef<ITableRow> STiledLandcapeImportDlg::OnGenerateWidgetForLayerDataListView(TSharedPtr<FLandscapeImportLayerData> InLayerData, const TSharedRef<STableViewBase>& OwnerTable)
+{
+	return 	SNew(STableRow< TSharedRef<FLandscapeImportLayerData> >, OwnerTable)
+			[
+				SNew(SBorder)
+				[
+					SNew(SHorizontalBox)
+				
+					+SHorizontalBox::Slot()
+					.VAlign(VAlign_Center)
+					.HAlign(HAlign_Left)
+					.FillWidth(1)
+					[
+						SNew(STextBlock).Text(InLayerData->LayerName)
+					]
+
+					+SHorizontalBox::Slot()
+					.VAlign(VAlign_Center)
+					.HAlign(HAlign_Right)
+					.Padding(2)
+					.AutoWidth()
+					[
+						SNew(STextBlock).Text(this, &STiledLandcapeImportDlg::GetWeightmapCountText, InLayerData)
+					]
+					
+					+SHorizontalBox::Slot()
+					.VAlign(VAlign_Center)
+					.HAlign(HAlign_Right)
+					.AutoWidth()
+					[
+						SNew(SButton)
+						.HAlign(HAlign_Center)
+						.ContentPadding(FEditorStyle::GetMargin("StandardDialog.ContentPadding"))
+						.OnClicked(this, &STiledLandcapeImportDlg::OnClickedSelectWeighmapTiles, InLayerData)
+						.Text(LOCTEXT("TiledLandscapeImport_SelectWeightmapButtonText", "Select Weightmap Tiles..."))
+					]
+				]
+			];
+	
+}
+
 const FTiledLandscapeImportSettings& STiledLandcapeImportDlg::GetImportSettings() const
 {
 	return ImportSettings;
@@ -294,14 +406,14 @@ void STiledLandcapeImportDlg::OnSetImportConfiguration(TSharedPtr<FTileImportCon
 	else
 	{
 		ImportSettings.ComponentsNum = 0;
-		ImportSettings.ImportFileList.Empty();
+		ImportSettings.HeightmapFileList.Empty();
 	}
 }
 
-FReply STiledLandcapeImportDlg::OnClickedSelectTiles()
+FReply STiledLandcapeImportDlg::OnClickedSelectHeightmapTiles()
 {
 	TotalLandscapeRect = FIntRect(MAX_int32, MAX_int32, MIN_int32, MIN_int32);
-	ImportSettings.ImportFileList.Empty();
+	ImportSettings.HeightmapFileList.Empty();
 	ImportSettings.TileCoordinates.Empty();
 	
 	SetPossibleConfigurationsForResolution(0);
@@ -314,22 +426,23 @@ FReply STiledLandcapeImportDlg::OnClickedSelectTiles()
 			void* ParentWindowWindowHandle = ParentWindow->GetNativeWindow()->GetOSWindowHandle();
 
 			const FString ImportFileTypes = TEXT("Raw heightmap tiles (*.r16)|*.r16");
-			if (DesktopPlatform->OpenFileDialog(
-				ParentWindowWindowHandle,
-				LOCTEXT("SelectHeightmapTiles", "Select heightmap tiles").ToString(),
-				*FEditorDirectories::Get().GetLastDirectory(ELastDirectory::UNR),
-				TEXT(""),
-				*ImportFileTypes,
-				EFileDialogFlags::Multiple,
-				ImportSettings.ImportFileList
-				) && ImportSettings.ImportFileList.Num())
+			bool bOpened = DesktopPlatform->OpenFileDialog(
+								ParentWindowWindowHandle,
+								LOCTEXT("SelectHeightmapTiles", "Select heightmap tiles").ToString(),
+								*FEditorDirectories::Get().GetLastDirectory(ELastDirectory::UNR),
+								TEXT(""),
+								*ImportFileTypes,
+								EFileDialogFlags::Multiple,
+								ImportSettings.HeightmapFileList);
+
+			if (bOpened	&& ImportSettings.HeightmapFileList.Num())
 			{
 				IFileManager& FileManager = IFileManager::Get();
 				bool bValidTiles = true;
 
 				// All heightmap tiles have to be the same size and have correct tile position encoded into filename
-				const int64 TargetFileSize = FileManager.FileSize(*ImportSettings.ImportFileList[0]);
-				for (const FString& Filename : ImportSettings.ImportFileList)
+				const int64 TargetFileSize = FileManager.FileSize(*ImportSettings.HeightmapFileList[0]);
+				for (const FString& Filename : ImportSettings.HeightmapFileList)
 				{
 					int64 FileSize = FileManager.FileSize(*Filename);
 					if (FileSize != TargetFileSize)
@@ -362,13 +475,58 @@ FReply STiledLandcapeImportDlg::OnClickedSelectTiles()
 	return FReply::Handled();
 }
 
+FReply STiledLandcapeImportDlg::OnClickedSelectWeighmapTiles(TSharedPtr<FLandscapeImportLayerData> InLayerData)
+{
+	InLayerData->WeighmapFileList.Empty();
+
+	IDesktopPlatform* DesktopPlatform = FDesktopPlatformModule::Get();
+	if (DesktopPlatform)
+	{
+		if (ParentWindow->GetNativeWindow().IsValid())
+		{
+			void* ParentWindowWindowHandle = ParentWindow->GetNativeWindow()->GetOSWindowHandle();
+
+			const FString ImportFileTypes = TEXT("Raw weigtmap tiles (*.raw)|*.raw");
+			bool bOpened = DesktopPlatform->OpenFileDialog(
+								ParentWindowWindowHandle,
+								LOCTEXT("SelectHeightmapTiles", "Select weightmap tiles").ToString(),
+								*FEditorDirectories::Get().GetLastDirectory(ELastDirectory::UNR),
+								TEXT(""),
+								*ImportFileTypes,
+								EFileDialogFlags::Multiple,
+								InLayerData->WeighmapFileList);
+
+			if (bOpened	&& InLayerData->WeighmapFileList.Num())
+			{
+				// TODO: check if it's a valid weightmaps
+			}
+		}
+	}
+	
+	return FReply::Handled();
+}
+
 bool STiledLandcapeImportDlg::IsImportEnabled() const
 {
-	return ImportSettings.ImportFileList.Num() && ImportSettings.ComponentsNum > 0;
+	return ImportSettings.HeightmapFileList.Num() && ImportSettings.ComponentsNum > 0;
 }
 
 FReply STiledLandcapeImportDlg::OnClickedImport()
 {
+	// copy weightmaps list data to an import structure  
+	ImportSettings.WeightmapFileList.Empty();
+	ImportSettings.WeightmapFileList.SetNum(ImportSettings.LandscapeLayerNameList.Num());
+	check(ImportSettings.WeightmapFileList.Num() == LayerDataList.Num());
+	for (int32 LayerIdx = 0; LayerIdx < LayerDataList.Num(); ++LayerIdx)
+	{
+		ImportSettings.WeightmapFileList[LayerIdx] = LayerDataList[LayerIdx]->WeighmapFileList;
+		
+		// we assume weightmaps tiles will match heightmap tiles 
+		ImportSettings.WeightmapFileList[LayerIdx].Sort();
+	}
+
+	ImportSettings.HeightmapFileList.Sort();
+			
 	ParentWindow->RequestDestroyWindow();
 	return FReply::Handled();	
 }
@@ -377,10 +535,19 @@ FReply STiledLandcapeImportDlg::OnClickedCancel()
 {
 	ParentWindow->RequestDestroyWindow();
 	
-	ImportSettings.ImportFileList.Empty();
+	ImportSettings.HeightmapFileList.Empty();
 	return FReply::Handled();	
 }
 
+void STiledLandcapeImportDlg::OnLandscapeMaterialChanged(const FAssetData& AssetData)
+{
+	ImportSettings.LandscapeMaterial = Cast<UMaterialInterface>(AssetData.GetAsset());
+	LandscapeMaterialComboButton->SetIsOpen(false);
+
+	// pull landscape layers from a chosen material
+	UpdateLandscapeLayerList();
+}
+
 void STiledLandcapeImportDlg::SetPossibleConfigurationsForResolution(int32 TargetResolutuion)
 {
 	int32 Idx = AllConfigurations.IndexOfByPredicate([&](const FTileImportConfiguration& A){
@@ -438,13 +605,13 @@ void STiledLandcapeImportDlg::GenerateAllPossibleTileConfigurations()
 
 FText STiledLandcapeImportDlg::GetImportSummaryText() const
 {
-	if (ImportSettings.ImportFileList.Num() == 0 || ImportSettings.ComponentsNum <= 0)
+	if (ImportSettings.HeightmapFileList.Num() == 0 || ImportSettings.ComponentsNum <= 0)
 	{
 		return FText();
 	}
 	
 	// Tile information(num, resolution)
-	const FString TilesSummary = FString::Printf(TEXT("%d - %dx%d"), ImportSettings.ImportFileList.Num(), ImportSettings.TileResolution, ImportSettings.TileResolution);
+	const FString TilesSummary = FString::Printf(TEXT("%d - %dx%d"), ImportSettings.HeightmapFileList.Num(), ImportSettings.TileResolution, ImportSettings.TileResolution);
 	
 	// Total landscape size(NxN km)
 	int32 WidthInTilesX = TotalLandscapeRect.Width() + 1;
@@ -460,6 +627,12 @@ FText STiledLandcapeImportDlg::GetImportSummaryText() const
 		);
 }
 
+FText STiledLandcapeImportDlg::GetWeightmapCountText(TSharedPtr<FLandscapeImportLayerData> InLayerData) const
+{
+	int32 NumWeighmaps = InLayerData.IsValid() ? InLayerData->WeighmapFileList.Num() : 0;
+	return FText::AsNumber(NumWeighmaps);
+}
+
 FText STiledLandcapeImportDlg::GenerateConfigurationText(int32 NumComponents, int32 NumSectionsPerComponent, int32 NumQuadsPerSection) const
 {
 	const FString ComponentsStr = FString::Printf(TEXT("%dx%d"), NumComponents, NumComponents);
@@ -474,4 +647,25 @@ FText STiledLandcapeImportDlg::GenerateConfigurationText(int32 NumComponents, in
 		);
 }
 
+void STiledLandcapeImportDlg::UpdateLandscapeLayerList()
+{
+	ImportSettings.LandscapeLayerNameList = ALandscapeProxy::GetLayersFromMaterial(ImportSettings.LandscapeMaterial);
+	//
+	ImportSettings.WeightmapFileList.Empty();
+	ImportSettings.WeightmapFileList.SetNum(ImportSettings.LandscapeLayerNameList.Num());
+
+	//	
+	LayerDataList.Empty();
+	for (FName LayerName  : ImportSettings.LandscapeLayerNameList)
+	{
+		TSharedPtr<FLandscapeImportLayerData> LayerData = MakeShareable(new FLandscapeImportLayerData());
+		LayerData->LayerName = LayerName.ToString();
+		LayerData->bBlend = true;
+
+		LayerDataList.Add(LayerData);
+	}
+
+	LayerDataListView->RequestListRefresh();
+}
+
 #undef LOCTEXT_NAMESPACE
diff --git a/Engine/Source/Editor/WorldBrowser/Private/Tiles/STiledLandscapeImportDlg.h b/Engine/Source/Editor/WorldBrowser/Private/Tiles/STiledLandscapeImportDlg.h
--- a/Engine/Source/Editor/WorldBrowser/Private/Tiles/STiledLandscapeImportDlg.h
+++ b/Engine/Source/Editor/WorldBrowser/Private/Tiles/STiledLandscapeImportDlg.h
@@ -6,6 +6,7 @@
 // 
 class STiledLandcapeImportDlg
 	: public SCompoundWidget
+	, public FGCObject
 {
 public:	
 	/** */		
@@ -18,15 +19,13 @@ public:
 		int32 NumQuadsPerSection;
 	};
 
-	struct FExistingLandscape
+	struct FLandscapeImportLayerData
 	{
-		FGuid	LandcapeGuid;	
-		FString LandcapeLabel;
-		int32	NumComponents;
-		int32	NumSectionsPerComponent;
-		int32	NumQuadsPerSection;
+		FString			LayerName;
+		TArray<FString> WeighmapFileList;
+		bool			bBlend;
 	};
-	
+
 	SLATE_BEGIN_ARGS( STiledLandcapeImportDlg )
 		{}
 	SLATE_END_ARGS()
@@ -37,8 +36,19 @@ public:
 	const FTiledLandscapeImportSettings& GetImportSettings() const;
 
 private:
+	/** FGCObject interface*/
+	virtual void AddReferencedObjects( FReferenceCollector& Collector ) OVERRIDE;
+
+	/** */
+	TSharedRef<SWidget> CreateLandscapeMaterialPicker();
+	FText GetLandscapeMaterialName() const;
+
+	/** */
 	TSharedRef<SWidget> HandleTileConfigurationComboBoxGenarateWidget(TSharedPtr<FTileImportConfiguration> InItem) const;
 	FText				GetTileConfigurationText() const;
+	
+	/** */
+	TSharedRef<ITableRow> OnGenerateWidgetForLayerDataListView(TSharedPtr<FLandscapeImportLayerData> InLayerData, const TSharedRef<STableViewBase>& OwnerTable);
 
 	/** */
 	TOptional<float> GetScaleX() const;
@@ -59,12 +69,17 @@ private:
 	bool IsImportEnabled() const;
 		
 	/** */
-	FReply OnClickedSelectTiles();
+	FReply OnClickedSelectHeightmapTiles();
+	FReply OnClickedSelectWeighmapTiles(TSharedPtr<FLandscapeImportLayerData> InLayerData);
 	FReply OnClickedImport();
 	FReply OnClickedCancel();
 
+	/** */
+	void OnLandscapeMaterialChanged(const FAssetData& AssetData);
+
 	/** */
 	FText GetImportSummaryText() const;
+	FText GetWeightmapCountText(TSharedPtr<FLandscapeImportLayerData> InLayerData) const; 
 
 	/** */
 	void SetPossibleConfigurationsForResolution(int32 TargetResolutuion);
@@ -75,15 +90,22 @@ private:
 	/** */
 	FText GenerateConfigurationText(int32 NumComponents, int32 NumSectionsPerComponent,	int32 NumQuadsPerSection) const;
 
+	/** */
+	void UpdateLandscapeLayerList();
+
 private:
 	/** */
 	TSharedPtr<SWindow> ParentWindow;
-	
+		
 	/** */
 	TSharedPtr<SComboBox<TSharedPtr<FTileImportConfiguration>>> TileConfigurationComboBox;
 	
 	/** */
-	TSharedPtr<SComboBox<TSharedPtr<FExistingLandscape>>> ExistingLandscapeActors;
+	TSharedPtr<SComboButton> LandscapeMaterialComboButton;
+
+	TSharedPtr<SListView<TSharedPtr<FLandscapeImportLayerData>>>	LayerDataListView;
+	TArray<TSharedPtr<FLandscapeImportLayerData>>					LayerDataList;		
+
 	
 	/** */
 	FTiledLandscapeImportSettings ImportSettings;
diff --git a/Engine/Source/Editor/WorldBrowser/Private/Tiles/WorldTileCollectionModel.cpp b/Engine/Source/Editor/WorldBrowser/Private/Tiles/WorldTileCollectionModel.cpp
--- a/Engine/Source/Editor/WorldBrowser/Private/Tiles/WorldTileCollectionModel.cpp
+++ b/Engine/Source/Editor/WorldBrowser/Private/Tiles/WorldTileCollectionModel.cpp
@@ -1236,7 +1236,7 @@ void FWorldTileCollectionModel::AddLandscapeProxy_Executed(FWorldTileModel::EWor
 }
 
 template<typename DataType>
-static bool ReadHeighmapFile(TArray<DataType>& Result, const TCHAR* Filename, uint32 Flags = 0)
+static bool ReadRawFile(TArray<DataType>& Result, const TCHAR* Filename, uint32 Flags = 0)
 {
 	auto Reader = TScopedPointer<FArchive>(IFileManager::Get().CreateFileReader(Filename, Flags));
 	if (!Reader)
@@ -1253,6 +1253,41 @@ static bool ReadHeighmapFile(TArray<DataType>& Result, const TCHAR* Filename, ui
 	return Reader->Close();
 }
 
+static ULandscapeLayerInfoObject* GetandscapeLayerInfoObject(FName LayerName, FString ContentPath)
+{
+	// Build default layer object name and package name
+	FString LayerObjectName = FString::Printf(TEXT("%s_LayerInfo"), *LayerName.ToString());
+	FString Path = ContentPath + TEXT("_sharedassets/");
+	if (Path.StartsWith("/Temp/"))
+	{
+		Path = FString("/Game/") + Path.RightChop(FString("/Temp/").Len());
+	}
+	
+	FString PackageName = Path + LayerObjectName;
+	UPackage* Package = FindPackage(nullptr, *PackageName);
+	if (Package == nullptr)
+	{
+		Package = CreatePackage(nullptr, *PackageName);
+	}
+
+	ULandscapeLayerInfoObject* LayerInfo = FindObject<ULandscapeLayerInfoObject>(Package, *LayerObjectName);
+	if (LayerInfo == nullptr)
+	{
+		LayerInfo = ConstructObject<ULandscapeLayerInfoObject>(ULandscapeLayerInfoObject::StaticClass(), Package, FName(*LayerObjectName), RF_Public | RF_Standalone | RF_Transactional);
+		// Notify the asset registry
+		FAssetRegistryModule::AssetCreated(LayerInfo);
+		// Mark the package dirty...
+		Package->MarkPackageDirty();
+		//
+		TArray<UPackage*> PackagesToSave; 
+		PackagesToSave.Add(Package);
+		FEditorFileUtils::PromptForCheckoutAndSave(PackagesToSave, false, false);
+	}
+		
+	LayerInfo->LayerName = LayerName;
+	return LayerInfo;
+}
+
 void FWorldTileCollectionModel::ImportTiledLandscape_Executed()
 {
 	/** Create the window to host widget */
@@ -1271,13 +1306,13 @@ void FWorldTileCollectionModel::ImportTiledLandscape_Executed()
 
 	const FTiledLandscapeImportSettings& ImportSettings = ImportDialog->GetImportSettings();
 		
-	if (ImportSettings.ImportFileList.Num())
+	if (ImportSettings.HeightmapFileList.Num())
 	{
 		// Default path for imported landscape tiles
 		// Use tile prefix as a folder name under world root
 		FString WorldRootPath = FPackageName::LongPackageNameToFilename(GetWorld()->WorldComposition->GetWorldRoot());
 		// Extract tile prefix
-		FString FolderName = FPaths::GetBaseFilename(ImportSettings.ImportFileList[0]);
+		FString FolderName = FPaths::GetBaseFilename(ImportSettings.HeightmapFileList[0]);
 		int32 PrefixEnd = FolderName.Find(TEXT("_x"), ESearchCase::IgnoreCase, ESearchDir::FromEnd);
 		FolderName = FolderName.Left(PrefixEnd);
 		WorldRootPath+= FolderName;
@@ -1290,15 +1325,16 @@ void FWorldTileCollectionModel::ImportTiledLandscape_Executed()
 
 		// Import tiles
 		int32 TileIndex = 0;
-		for (const FString& Filename : ImportSettings.ImportFileList)
+		for (const FString& Filename : ImportSettings.HeightmapFileList)
 		{
 			FString TileName = FPaths::GetBaseFilename(Filename);
 			FVector TileScale = ImportSettings.Scale3D;
 
-			GWarn->StatusUpdate(TileIndex, ImportSettings.ImportFileList.Num(), FText::Format(LOCTEXT("ImportingLandscapeTiles", "Importing landscape tiles: {0}"), FText::FromString(TileName)));
+			GWarn->StatusUpdate(TileIndex, ImportSettings.HeightmapFileList.Num(), FText::Format(LOCTEXT("ImportingLandscapeTiles", "Importing landscape tiles: {0}"), FText::FromString(TileName)));
 			
 			FWorldTileModel::FLandscapeImportSettings TileImportSettings;
 			TileImportSettings.LandscapeGuid		= LandscapeGuid;
+			TileImportSettings.LandscapeMaterial	= ImportSettings.LandscapeMaterial;
 			TileImportSettings.ComponentSizeQuads	= ImportSettings.QuadsPerSection*ImportSettings.SectionsPerComponent;
 			TileImportSettings.QuadsPerSection		= ImportSettings.QuadsPerSection;
 			TileImportSettings.SectionsPerComponent = ImportSettings.SectionsPerComponent;
@@ -1307,7 +1343,24 @@ void FWorldTileCollectionModel::ImportTiledLandscape_Executed()
 			TileImportSettings.HeightmapFilename	= Filename;
 			TileImportSettings.LandscapeTransform.SetScale3D(TileScale);
 
-			if (ReadHeighmapFile(TileImportSettings.HeightData, *Filename, FILEREAD_Silent))
+			// Setup weightmaps for each layer
+			for (int32 LayerIdx = 0; LayerIdx < ImportSettings.LandscapeLayerNameList.Num(); ++LayerIdx)
+			{
+				FName LayerName = ImportSettings.LandscapeLayerNameList[LayerIdx];
+				
+				TileImportSettings.ImportLayers.Add(FLandscapeImportLayerInfo(LayerName));
+				FLandscapeImportLayerInfo& LayerImportInfo = TileImportSettings.ImportLayers.Last();
+
+				if (ImportSettings.WeightmapFileList[LayerIdx].IsValidIndex(TileIndex))
+				{
+					LayerImportInfo.SourceFilePath = ImportSettings.WeightmapFileList[LayerIdx][TileIndex];
+					ReadRawFile(LayerImportInfo.LayerData, *LayerImportInfo.SourceFilePath, FILEREAD_Silent);
+					LayerImportInfo.LayerInfo = GetandscapeLayerInfoObject(LayerImportInfo.LayerName, GetWorld()->GetOutermost()->GetName());
+					LayerImportInfo.LayerInfo->bNoWeightBlend = false; //option ?
+				}
+			}
+						
+			if (ReadRawFile(TileImportSettings.HeightData, *Filename, FILEREAD_Silent))
 			{
 				FString MapFileName = WorldRootPath + TileName + FPackageName::GetMapPackageExtension();
 				// Create a new world - so we can 'borrow' its level
@@ -1334,6 +1387,20 @@ void FWorldTileCollectionModel::ImportTiledLandscape_Executed()
 
 					if (NewLandscape)
 					{
+						// update shared LandscapeInfo objects with new layers
+						ULandscapeInfo* LandscapeInfo = NewLandscape->GetLandscapeInfo(true);
+						LandscapeInfo->UpdateLayerInfoMap(NewLandscape);
+						for (const FLandscapeImportLayerInfo& LayerImportInfo : TileImportSettings.ImportLayers)
+						{
+							NewLandscape->EditorLayerSettings.Add(FLandscapeEditorLayerSettings(LayerImportInfo.LayerInfo, LayerImportInfo.SourceFilePath));
+							int32 LayerInfoIndex = LandscapeInfo->GetLayerInfoIndex(LayerImportInfo.LayerInfo->LayerName);
+							if (ensure(LayerInfoIndex != INDEX_NONE))
+							{
+								FLandscapeInfoLayerSettings& LayerSettings = LandscapeInfo->Layers[LayerInfoIndex];
+								LayerSettings.LayerInfoObj = LayerImportInfo.LayerInfo;
+							}
+						}
+						
 						// Set landscape actor location at min corner of bounding rect, so it will be centered around origin
 						FIntRect NewLandscapeRect = NewLandscape->GetBoundingRect();
 						float WidthX = NewLandscapeRect.Width()*TileScale.X;
diff --git a/Engine/Source/Editor/WorldBrowser/Private/Tiles/WorldTileModel.cpp b/Engine/Source/Editor/WorldBrowser/Private/Tiles/WorldTileModel.cpp
--- a/Engine/Source/Editor/WorldBrowser/Private/Tiles/WorldTileModel.cpp
+++ b/Engine/Source/Editor/WorldBrowser/Private/Tiles/WorldTileModel.cpp
@@ -920,6 +920,11 @@ ALandscapeProxy* FWorldTileModel::ImportLandscape(const FLandscapeImportSettings
 	
 	Landscape->SetActorTransform(Settings.LandscapeTransform);
 	
+	if (Settings.LandscapeMaterial)
+	{
+		Landscape->LandscapeMaterial = Settings.LandscapeMaterial;
+	}
+	
 	// Create landscape components	
 	Landscape->Import(
 		LandscapeGuid, 
diff --git a/Engine/Source/Editor/WorldBrowser/Private/Tiles/WorldTileModel.h b/Engine/Source/Editor/WorldBrowser/Private/Tiles/WorldTileModel.h
--- a/Engine/Source/Editor/WorldBrowser/Private/Tiles/WorldTileModel.h
+++ b/Engine/Source/Editor/WorldBrowser/Private/Tiles/WorldTileModel.h
@@ -39,6 +39,7 @@ public:
 		// Depending on landscape guid import code will spawn Landscape actor or LandscapeProxy actor
 		FGuid								LandscapeGuid;
 		FTransform							LandscapeTransform;
+		UMaterialInterface*					LandscapeMaterial;
 		int32								ComponentSizeQuads;
 		int32								SectionsPerComponent;
 		int32								QuadsPerSection;
