World Browser: - Added scale ruler in top view

diff --git a/Engine/Source/Editor/WorldBrowser/Private/SWorldGridView.cpp b/Engine/Source/Editor/WorldBrowser/Private/SWorldGridView.cpp
--- a/Engine/Source/Editor/WorldBrowser/Private/SWorldGridView.cpp
+++ b/Engine/Source/Editor/WorldBrowser/Private/SWorldGridView.cpp
@@ -35,7 +35,7 @@ struct FWorldZoomLevelsContainer
 
 	int32	GetNumZoomLevels() const OVERRIDE
 	{
-		return 60;
+		return 100;
 	}
 
 	int32	GetDefaultZoomLevel() const OVERRIDE
@@ -406,6 +406,49 @@ uint32 SWorldLevelsGridView::PaintBackground(const FGeometry& AllottedGeometry,
 	return LayerId + 1;
 }
 
+uint32 SWorldLevelsGridView::PaintScaleRuler(const FGeometry& AllottedGeometry, const FSlateRect& MyClippingRect, FSlateWindowElementList& OutDrawElements, uint32 LayerId) const
+{
+	const float	ScaleRulerLength = 100.f; // pixels
+	TArray<FVector2D> LinePoints;
+	LinePoints.Add(FVector2D::ZeroVector);
+	LinePoints.Add(FVector2D::ZeroVector + FVector2D(ScaleRulerLength, 0.f));
+	
+	FSlateDrawElement::MakeLines( 
+		OutDrawElements,
+		LayerId,
+		AllottedGeometry.ToOffsetPaintGeometry(FVector2D(10, 40)),
+		LinePoints,
+		MyClippingRect,
+		ESlateDrawEffect::None,
+		FColor(200, 200, 200));
+
+	const float UnitsInRuler = ScaleRulerLength/GetZoomAmount();// Pixels to world units
+	const int32 UnitsInMeter = 100;
+	const int32 UnitsInKilometer = UnitsInMeter*1000;
+	
+	FString RulerText;
+	if (UnitsInRuler > UnitsInKilometer) // in kilometers
+	{
+		RulerText = FString::Printf(TEXT("%.2f km"), UnitsInRuler/UnitsInKilometer);
+	}
+	else // in meters
+	{
+		RulerText = FString::Printf(TEXT("%.2f m"), UnitsInRuler/UnitsInMeter);
+	}
+	
+	FSlateDrawElement::MakeText(
+		OutDrawElements,
+		LayerId,
+		AllottedGeometry.ToOffsetPaintGeometry(FVector2D(10, 27)),
+		RulerText,
+		FEditorStyle::GetFontStyle("NormalFont"),
+		MyClippingRect,
+		ESlateDrawEffect::None,
+		FColor(200, 200, 200));
+	
+	return LayerId + 1;
+}
+
 int32 SWorldLevelsGridView::OnPaint( const FGeometry& AllottedGeometry, const FSlateRect& MyClippingRect, 
 							FSlateWindowElementList& OutDrawElements, int32 LayerId, 
 							const FWidgetStyle& InWidgetStyle, bool bParentEnabled ) const
@@ -501,6 +544,7 @@ int32 SWorldLevelsGridView::OnPaint( const FGeometry& AllottedGeometry, const FS
 		}
 	}
 
+	LayerId = PaintScaleRuler(AllottedGeometry, MyClippingRect, OutDrawElements, LayerId);
 	return LayerId;
 }
 
diff --git a/Engine/Source/Editor/WorldBrowser/Private/SWorldGridView.h b/Engine/Source/Editor/WorldBrowser/Private/SWorldGridView.h
--- a/Engine/Source/Editor/WorldBrowser/Private/SWorldGridView.h
+++ b/Engine/Source/Editor/WorldBrowser/Private/SWorldGridView.h
@@ -56,6 +56,9 @@ public:
 protected:
 	/**  Draws background for grid view */
 	uint32 PaintBackground(const FGeometry& AllottedGeometry, const FSlateRect& MyClippingRect, FSlateWindowElementList& OutDrawElements, uint32 LayerId) const;
+
+	/**  Draws current scale */
+	uint32 PaintScaleRuler(const FGeometry& AllottedGeometry, const FSlateRect& MyClippingRect, FSlateWindowElementList& OutDrawElements, uint32 LayerId) const;
 	
 	/**  SWidget interface */
 	FReply OnKeyDown(const FGeometry& MyGeometry, const FKeyboardEvent& InKeyboardEvent);
diff --git a/Engine/Source/Editor/WorldBrowser/Private/SWorldMainView.cpp b/Engine/Source/Editor/WorldBrowser/Private/SWorldMainView.cpp
--- a/Engine/Source/Editor/WorldBrowser/Private/SWorldMainView.cpp
+++ b/Engine/Source/Editor/WorldBrowser/Private/SWorldMainView.cpp
@@ -321,15 +321,12 @@ TSharedRef<SWidget> SWorldMainView::ConstructCenterPane()
 						+SHorizontalBox::Slot()
 						.AutoWidth()
 						[
-							SNew(SImage)
-							.Image(FEditorStyle::GetBrush( "WorldBrowser.WorldViewScale" ))
+							SNullWidget::NullWidget
 						]
 						+SHorizontalBox::Slot()
 						.Padding(5,0,0,0)
 						[
-							SNew(STextBlock)
-							.TextStyle( FEditorStyle::Get(), "WorldBrowser.StatusBarText" )
-							.Text(this, &SWorldMainView::GetZoomText)
+							SNullWidget::NullWidget
 						]
 							
 						// World origin position
