World Browser: - Fixed unnecessary GC reference handling

diff --git a/Engine/Source/Editor/WorldBrowser/Private/LevelCollectionModel.h b/Engine/Source/Editor/WorldBrowser/Private/LevelCollectionModel.h
--- a/Engine/Source/Editor/WorldBrowser/Private/LevelCollectionModel.h
+++ b/Engine/Source/Editor/WorldBrowser/Private/LevelCollectionModel.h
@@ -522,7 +522,6 @@ struct FTiledLandscapeImportSettings
 		, SectionsPerComponent(1)
 		, TilesCoordinatesOffset(0,0)
 		, TileResolution(1009)
-		, LandscapeMaterial(nullptr)
 	{}
 	
 	FVector				Scale3D;
@@ -536,7 +535,7 @@ struct FTiledLandscapeImportSettings
 	int32				TileResolution;	
 
 
-	UMaterialInterface*					LandscapeMaterial;
+	TWeakObjectPtr<UMaterialInterface>	LandscapeMaterial;
 	
 	TArray<FName>						LandscapeLayerNameList;
 	// list of weightmap files per each layer
diff --git a/Engine/Source/Editor/WorldBrowser/Private/Tiles/STiledLandscapeImportDlg.cpp b/Engine/Source/Editor/WorldBrowser/Private/Tiles/STiledLandscapeImportDlg.cpp
--- a/Engine/Source/Editor/WorldBrowser/Private/Tiles/STiledLandscapeImportDlg.cpp
+++ b/Engine/Source/Editor/WorldBrowser/Private/Tiles/STiledLandscapeImportDlg.cpp
@@ -38,11 +38,6 @@ static FIntPoint ExtractHeighmapTileCoordinates(FString BaseFilename)
 	return ResultPosition;
 }
 
-void STiledLandcapeImportDlg::AddReferencedObjects( FReferenceCollector& Collector )
-{
-	Collector.AddReferencedObject(ImportSettings.LandscapeMaterial);
-}
-
 void STiledLandcapeImportDlg::Construct(const FArguments& InArgs, TSharedPtr<SWindow> InParentWindow)
 {
 	ParentWindow = InParentWindow;
@@ -267,12 +262,12 @@ TSharedRef<SWidget> STiledLandcapeImportDlg::CreateLandscapeMaterialPicker()
 
 FText STiledLandcapeImportDlg::GetLandscapeMaterialName() const
 {
-	if (ImportSettings.LandscapeMaterial)
+	if (ImportSettings.LandscapeMaterial.IsValid())
 	{
 		return FText::FromString(ImportSettings.LandscapeMaterial->GetName());
 	}
 
-	return LOCTEXT("TiledLandscapeImport_NoLandscapeMaterialText", "Default Material");
+	return LOCTEXT("TiledLandscapeImport_NoLandscapeMaterialText", "None");
 }
 
 TSharedRef<SWidget> STiledLandcapeImportDlg::HandleTileConfigurationComboBoxGenarateWidget(TSharedPtr<FTileImportConfiguration> InItem) const
@@ -649,7 +644,7 @@ FText STiledLandcapeImportDlg::GenerateConfigurationText(int32 NumComponents, in
 
 void STiledLandcapeImportDlg::UpdateLandscapeLayerList()
 {
-	ImportSettings.LandscapeLayerNameList = ALandscapeProxy::GetLayersFromMaterial(ImportSettings.LandscapeMaterial);
+	ImportSettings.LandscapeLayerNameList = ALandscapeProxy::GetLayersFromMaterial(ImportSettings.LandscapeMaterial.Get());
 	//
 	ImportSettings.WeightmapFileList.Empty();
 	ImportSettings.WeightmapFileList.SetNum(ImportSettings.LandscapeLayerNameList.Num());
diff --git a/Engine/Source/Editor/WorldBrowser/Private/Tiles/STiledLandscapeImportDlg.h b/Engine/Source/Editor/WorldBrowser/Private/Tiles/STiledLandscapeImportDlg.h
--- a/Engine/Source/Editor/WorldBrowser/Private/Tiles/STiledLandscapeImportDlg.h
+++ b/Engine/Source/Editor/WorldBrowser/Private/Tiles/STiledLandscapeImportDlg.h
@@ -6,7 +6,6 @@
 // 
 class STiledLandcapeImportDlg
 	: public SCompoundWidget
-	, public FGCObject
 {
 public:	
 	/** */		
@@ -36,9 +35,6 @@ public:
 	const FTiledLandscapeImportSettings& GetImportSettings() const;
 
 private:
-	/** FGCObject interface*/
-	virtual void AddReferencedObjects( FReferenceCollector& Collector ) OVERRIDE;
-
 	/** */
 	TSharedRef<SWidget> CreateLandscapeMaterialPicker();
 	FText GetLandscapeMaterialName() const;
@@ -105,8 +101,7 @@ private:
 
 	TSharedPtr<SListView<TSharedPtr<FLandscapeImportLayerData>>>	LayerDataListView;
 	TArray<TSharedPtr<FLandscapeImportLayerData>>					LayerDataList;		
-
-	
+		
 	/** */
 	FTiledLandscapeImportSettings ImportSettings;
 	
diff --git a/Engine/Source/Editor/WorldBrowser/Private/Tiles/WorldTileCollectionModel.cpp b/Engine/Source/Editor/WorldBrowser/Private/Tiles/WorldTileCollectionModel.cpp
--- a/Engine/Source/Editor/WorldBrowser/Private/Tiles/WorldTileCollectionModel.cpp
+++ b/Engine/Source/Editor/WorldBrowser/Private/Tiles/WorldTileCollectionModel.cpp
@@ -1334,7 +1334,7 @@ void FWorldTileCollectionModel::ImportTiledLandscape_Executed()
 			
 			FWorldTileModel::FLandscapeImportSettings TileImportSettings;
 			TileImportSettings.LandscapeGuid		= LandscapeGuid;
-			TileImportSettings.LandscapeMaterial	= ImportSettings.LandscapeMaterial;
+			TileImportSettings.LandscapeMaterial	= ImportSettings.LandscapeMaterial.Get();
 			TileImportSettings.ComponentSizeQuads	= ImportSettings.QuadsPerSection*ImportSettings.SectionsPerComponent;
 			TileImportSettings.QuadsPerSection		= ImportSettings.QuadsPerSection;
 			TileImportSettings.SectionsPerComponent = ImportSettings.SectionsPerComponent;
