Disabled "Merge..." command when persistent level is selected in World Browser

diff --git a/Engine/Source/Editor/WorldBrowser/Private/LevelCollectionModel.cpp b/Engine/Source/Editor/WorldBrowser/Private/LevelCollectionModel.cpp
--- a/Engine/Source/Editor/WorldBrowser/Private/LevelCollectionModel.cpp
+++ b/Engine/Source/Editor/WorldBrowser/Private/LevelCollectionModel.cpp
@@ -816,6 +816,19 @@ bool FLevelCollectionModel::AreAllSelectedLevelsEditable() const
 	return AreAnyLevelsSelected();
 }
 
+bool FLevelCollectionModel::AreAllSelectedLevelsEditableAndNotPersistent() const
+{
+	for (auto It = SelectedLevelsList.CreateConstIterator(); It; ++It)
+	{
+		if ((*It)->IsEditable() == false || (*It)->IsPersistent())
+		{
+			return false;
+		}
+	}
+	
+	return AreAnyLevelsSelected();
+}
+
 bool FLevelCollectionModel::AreAllSelectedLevelsEditableAndVisible() const
 {
 	for (auto It = SelectedLevelsList.CreateConstIterator(); It; ++It)
diff --git a/Engine/Source/Editor/WorldBrowser/Private/LevelCollectionModel.h b/Engine/Source/Editor/WorldBrowser/Private/LevelCollectionModel.h
--- a/Engine/Source/Editor/WorldBrowser/Private/LevelCollectionModel.h
+++ b/Engine/Source/Editor/WorldBrowser/Private/LevelCollectionModel.h
@@ -181,6 +181,9 @@ public:
 	/** @return whether all the currently selected levels are editable */
 	bool AreAllSelectedLevelsEditable() const;
 
+	/** @return whether all the currently selected levels are editable and not persistent */
+	bool AreAllSelectedLevelsEditableAndNotPersistent() const;
+
 	/** @return whether all the currently selected levels are editable and visible*/
 	bool AreAllSelectedLevelsEditableAndVisible() const;
 
diff --git a/Engine/Source/Editor/WorldBrowser/Private/StreamingLevels/StreamingLevelCollectionModel.cpp b/Engine/Source/Editor/WorldBrowser/Private/StreamingLevels/StreamingLevelCollectionModel.cpp
--- a/Engine/Source/Editor/WorldBrowser/Private/StreamingLevels/StreamingLevelCollectionModel.cpp
+++ b/Engine/Source/Editor/WorldBrowser/Private/StreamingLevels/StreamingLevelCollectionModel.cpp
@@ -159,7 +159,7 @@ void FStreamingLevelCollectionModel::BindCommands()
 	
 	ActionList.MapAction( Commands.World_MergeSelectedLevels,
 		FExecuteAction::CreateSP( this, &FStreamingLevelCollectionModel::MergeSelectedLevels_Executed  ),
-		FCanExecuteAction::CreateSP( this, &FStreamingLevelCollectionModel::AreAllSelectedLevelsEditable ) );
+		FCanExecuteAction::CreateSP( this, &FStreamingLevelCollectionModel::AreAllSelectedLevelsEditableAndNotPersistent ) );
 	
 	// new level streaming method
 	ActionList.MapAction( Commands.SetAddStreamingMethod_Blueprint,
@@ -565,6 +565,11 @@ void FStreamingLevelCollectionModel::RemoveInvalidSelectedLevels_Executed()
 
 void FStreamingLevelCollectionModel::MergeSelectedLevels_Executed()
 {
+	if (SelectedLevelsList.Num() <= 1)
+	{
+		return;
+	}
+	
 	// Stash off a copy of the original array, so the selection can be restored
 	FLevelModelList SelectedLevelsCopy = SelectedLevelsList;
 
