Fixed landscape importing crash in the World Browser

diff --git a/Engine/Source/Editor/WorldBrowser/Private/Tiles/WorldTileCollectionModel.cpp b/Engine/Source/Editor/WorldBrowser/Private/Tiles/WorldTileCollectionModel.cpp
--- a/Engine/Source/Editor/WorldBrowser/Private/Tiles/WorldTileCollectionModel.cpp
+++ b/Engine/Source/Editor/WorldBrowser/Private/Tiles/WorldTileCollectionModel.cpp
@@ -1332,7 +1332,7 @@ void FWorldTileCollectionModel::ImportTiledLandscape_Executed()
 
 			GWarn->StatusUpdate(TileIndex, ImportSettings.HeightmapFileList.Num(), FText::Format(LOCTEXT("ImportingLandscapeTiles", "Importing landscape tiles: {0}"), FText::FromString(TileName)));
 			
-			FWorldTileModel::FLandscapeImportSettings TileImportSettings;
+			FWorldTileModel::FLandscapeImportSettings TileImportSettings = {};
 			TileImportSettings.LandscapeGuid		= LandscapeGuid;
 			TileImportSettings.LandscapeMaterial	= ImportSettings.LandscapeMaterial.Get();
 			TileImportSettings.ComponentSizeQuads	= ImportSettings.QuadsPerSection*ImportSettings.SectionsPerComponent;
diff --git a/Engine/Source/Editor/WorldBrowser/Private/Tiles/WorldTileModel.cpp b/Engine/Source/Editor/WorldBrowser/Private/Tiles/WorldTileModel.cpp
--- a/Engine/Source/Editor/WorldBrowser/Private/Tiles/WorldTileModel.cpp
+++ b/Engine/Source/Editor/WorldBrowser/Private/Tiles/WorldTileModel.cpp
@@ -837,7 +837,7 @@ bool FWorldTileModel::CreateAdjacentLandscapeProxy(ALandscapeProxy* SourceLandsc
 	FIntRect SourceLandscapeRect = SourceLandscape->GetBoundingRect();
 	FIntPoint SourceLandscapeSize = SourceLandscapeRect.Size();
 
-	FLandscapeImportSettings ImportSettings;
+	FLandscapeImportSettings ImportSettings = {};
 	ImportSettings.LandscapeGuid = SourceLandscape->GetLandscapeGuid();
 	ImportSettings.ComponentSizeQuads = SourceLandscape->ComponentSizeQuads;
 	ImportSettings.SectionsPerComponent = SourceLandscape->NumSubsections;
