Fixed case when "Assign To Layer" command does not appear in World Browser hierarchy view

diff --git a/Engine/Source/Editor/WorldBrowser/Private/Tiles/WorldTileCollectionModel.cpp b/Engine/Source/Editor/WorldBrowser/Private/Tiles/WorldTileCollectionModel.cpp
--- a/Engine/Source/Editor/WorldBrowser/Private/Tiles/WorldTileCollectionModel.cpp
+++ b/Engine/Source/Editor/WorldBrowser/Private/Tiles/WorldTileCollectionModel.cpp
@@ -377,9 +377,8 @@ void FWorldTileCollectionModel::BuildHierarchyMenu(FMenuBuilder& MenuBuilder) co
 		}
 		MenuBuilder.EndSection();
 	}
-
-	if (AreAnyLevelsSelected() && 
-		StaticCastSharedPtr<FWorldTileModel>(SelectedLevelsList[0])->IsLandscapeBased())
+	
+	if (AreAnySelectedLevelsEditable())
 	{
 		MenuBuilder.BeginSection("Menu_Layers");
 		{
