Fixed crash on importing tiled landcape when layers has no associated weightmaps

diff --git a/Engine/Source/Editor/WorldBrowser/Private/Tiles/WorldTileCollectionModel.cpp b/Engine/Source/Editor/WorldBrowser/Private/Tiles/WorldTileCollectionModel.cpp
--- a/Engine/Source/Editor/WorldBrowser/Private/Tiles/WorldTileCollectionModel.cpp
+++ b/Engine/Source/Editor/WorldBrowser/Private/Tiles/WorldTileCollectionModel.cpp
@@ -1391,12 +1391,15 @@ void FWorldTileCollectionModel::ImportTiledLandscape_Executed()
 						LandscapeInfo->UpdateLayerInfoMap(NewLandscape);
 						for (const FLandscapeImportLayerInfo& LayerImportInfo : TileImportSettings.ImportLayers)
 						{
-							NewLandscape->EditorLayerSettings.Add(FLandscapeEditorLayerSettings(LayerImportInfo.LayerInfo, LayerImportInfo.SourceFilePath));
-							int32 LayerInfoIndex = LandscapeInfo->GetLayerInfoIndex(LayerImportInfo.LayerInfo->LayerName);
-							if (ensure(LayerInfoIndex != INDEX_NONE))
+							if (LayerImportInfo.LayerInfo)
 							{
-								FLandscapeInfoLayerSettings& LayerSettings = LandscapeInfo->Layers[LayerInfoIndex];
-								LayerSettings.LayerInfoObj = LayerImportInfo.LayerInfo;
+								NewLandscape->EditorLayerSettings.Add(FLandscapeEditorLayerSettings(LayerImportInfo.LayerInfo, LayerImportInfo.SourceFilePath));
+								int32 LayerInfoIndex = LandscapeInfo->GetLayerInfoIndex(LayerImportInfo.LayerInfo->LayerName);
+								if (ensure(LayerInfoIndex != INDEX_NONE))
+								{
+									FLandscapeInfoLayerSettings& LayerSettings = LandscapeInfo->Layers[LayerInfoIndex];
+									LayerSettings.LayerInfoObj = LayerImportInfo.LayerInfo;
+								}
 							}
 						}
 						
