World Browser UX improvements
Splitted World Browser views into three independent windows, each window can be summoned from a hierarchy view or details view toolbar
World composition can now be toggled through WorldSettings - bEnableWorldComposition. Once enabled persistent level will be treated as a world root. No more '?worldcompistion' parameter workarounds. 'Open World' command was removed, as levels saved with world composition enabled can be open through 'Open Level' command.
Added bEnableWorldOriginRebasing parameter to WorldSettings, which contols whether world origin shifting should be used

diff --git a/Engine/Source/Developer/AssetTools/Private/AssetTypeActions/AssetTypeActions_World.cpp b/Engine/Source/Developer/AssetTools/Private/AssetTypeActions/AssetTypeActions_World.cpp
--- a/Engine/Source/Developer/AssetTools/Private/AssetTypeActions/AssetTypeActions_World.cpp
+++ b/Engine/Source/Developer/AssetTools/Private/AssetTypeActions/AssetTypeActions_World.cpp
@@ -14,8 +14,7 @@ void FAssetTypeActions_World::OpenAssetEditor( const TArray<UObject*>& InObjects
 			const FString FileToOpen = FPackageName::LongPackageNameToFilename(World->GetOutermost()->GetName(), FPackageName::GetMapPackageExtension());
 			const bool bLoadAsTemplate = false;
 			const bool bShowProgress = false;
-			const bool bWorldComposition = false;
-			FEditorFileUtils::LoadMap( FileToOpen, bLoadAsTemplate, bShowProgress, bWorldComposition );
+			FEditorFileUtils::LoadMap( FileToOpen, bLoadAsTemplate, bShowProgress );
 
 			// We can only edit one world at a time... so just break after the first valid world to load
 			break;
diff --git a/Engine/Source/Editor/EditorStyle/Private/SlateEditorStyle.cpp b/Engine/Source/Editor/EditorStyle/Private/SlateEditorStyle.cpp
--- a/Engine/Source/Editor/EditorStyle/Private/SlateEditorStyle.cpp
+++ b/Engine/Source/Editor/EditorStyle/Private/SlateEditorStyle.cpp
@@ -1052,7 +1052,10 @@ void FSlateEditorStyle::FStyle::SetupGeneralStyles()
 		Set( "WorldBrowser.DirectionYNegative", new IMAGE_BRUSH( "Icons/icon_PanDown", Icon16x16 ) );
 		Set( "WorldBrowser.LevelStreamingAlwaysLoaded", new FSlateNoResource() );
 		Set( "WorldBrowser.LevelStreamingKismet", new IMAGE_BRUSH( "Icons/ActorIcons/ComponentMobilityStationary_7x16px", Icon7x16 ) );
-
+		Set( "WorldBrowser.HierarchyButtonBrush", new IMAGE_BRUSH( "Icons/levels_16x", Icon16x16 ) );
+		Set( "WorldBrowser.DetailsButtonBrush", new IMAGE_BRUSH( "NewsFeed/SettingsButton", Icon16x16 ) );
+		Set( "WorldBrowser.CompositionButtonBrush", new IMAGE_BRUSH( "Icons/icon_volume_LevelGrid_16x", Icon16x16 ) );
+		
 		Set( "WorldBrowser.StatusBarText", FTextBlockStyle(NormalText)
 				.SetFont( TTF_FONT( "Fonts/Roboto-BoldCondensed", 12 ) )
 				.SetColorAndOpacity( FLinearColor(0.9, 0.9f, 0.9f, 0.5f) )
diff --git a/Engine/Source/Editor/LevelEditor/Private/LevelEditor.cpp b/Engine/Source/Editor/LevelEditor/Private/LevelEditor.cpp
--- a/Engine/Source/Editor/LevelEditor/Private/LevelEditor.cpp
+++ b/Engine/Source/Editor/LevelEditor/Private/LevelEditor.cpp
@@ -234,6 +234,24 @@ void FLevelEditorModule::SummonLevelBrowser()
 	LevelEditorInstance->InvokeTab("LevelEditorLevelBrowser");
 }
 
+void FLevelEditorModule::SummonWorldBrowserHierarchy()
+{
+	TSharedPtr<SLevelEditor> LevelEditorInstance = LevelEditorInstancePtr.Pin();
+	LevelEditorInstance->InvokeTab("WorldBrowserHierarchy");
+}
+
+void FLevelEditorModule::SummonWorldBrowserDetails()
+{
+	TSharedPtr<SLevelEditor> LevelEditorInstance = LevelEditorInstancePtr.Pin();
+	LevelEditorInstance->InvokeTab("WorldBrowserDetails");
+}
+
+void FLevelEditorModule::SummonWorldBrowserComposition()
+{
+	TSharedPtr<SLevelEditor> LevelEditorInstance = LevelEditorInstancePtr.Pin();
+	LevelEditorInstance->InvokeTab("WorldBrowserComposition");
+}
+
 // @todo remove when world-centric mode is added
 void FLevelEditorModule::AttachSequencer(TSharedPtr<SWidget> Sequencer)
 {
diff --git a/Engine/Source/Editor/LevelEditor/Private/SLevelEditor.cpp b/Engine/Source/Editor/LevelEditor/Private/SLevelEditor.cpp
--- a/Engine/Source/Editor/LevelEditor/Private/SLevelEditor.cpp
+++ b/Engine/Source/Editor/LevelEditor/Private/SLevelEditor.cpp
@@ -41,10 +41,12 @@
 
 static const FName LevelEditorBuildAndSubmitTab("LevelEditorBuildAndSubmit");
 static const FName LevelEditorStatsViewerTab("LevelEditorStatsViewer");
-static const FName LevelEditorWorldBrowserTab("LevelEditorWorldBrowser");
 static const FName MainFrameModuleName("MainFrame");
 static const FName NewsFeedModuleName("NewsFeed");
 static const FName LevelEditorModuleName("LevelEditor");
+static const FName WorldBrowserHierarchyTab("WorldBrowserHierarchy");
+static const FName WorldBrowserDetailsTab("WorldBrowserDetails");
+static const FName WorldBrowserCompositionTab("WorldBrowserComposition");
 
 
 namespace LevelEditorConstants
@@ -691,14 +693,34 @@ TSharedRef<SDockTab> SLevelEditor::SpawnLevelEditorTab( const FSpawnTabArgs& Arg
 					]
 				];
 	}
-	else if( GetDefault<UEditorExperimentalSettings>()->bWorldBrowser && TabIdentifier == LevelEditorWorldBrowserTab )
+	else if( GetDefault<UEditorExperimentalSettings>()->bWorldBrowser && TabIdentifier == WorldBrowserHierarchyTab)
 	{
 		FWorldBrowserModule& WorldBrowserModule = FModuleManager::LoadModuleChecked<FWorldBrowserModule>( "WorldBrowser" );
 		return SNew( SDockTab )
 			.Icon( FEditorStyle::GetBrush( "LevelEditor.Tabs.WorldBrowser" ) )
-			.Label( NSLOCTEXT("LevelEditor", "WorldBrowserTabTitle", "World Browser") )
+			.Label( NSLOCTEXT("LevelEditor", "WorldBrowserHierarchyTabTitle", "Levels Hierarchy") )
 			[
-				WorldBrowserModule.CreateWorldBrowser()
+				WorldBrowserModule.CreateWorldBrowserHierarchy()
+			];
+	}
+	else if( GetDefault<UEditorExperimentalSettings>()->bWorldBrowser && TabIdentifier == WorldBrowserDetailsTab)
+	{
+		FWorldBrowserModule& WorldBrowserModule = FModuleManager::LoadModuleChecked<FWorldBrowserModule>( "WorldBrowser" );
+		return SNew( SDockTab )
+			.Icon( FEditorStyle::GetBrush( "LevelEditor.Tabs.WorldBrowser" ) )
+			.Label( NSLOCTEXT("LevelEditor", "WorldBrowserDetailsTabTitle", "Levels Details") )
+			[
+				WorldBrowserModule.CreateWorldBrowserDetails()
+			];
+	}
+	else if( GetDefault<UEditorExperimentalSettings>()->bWorldBrowser && TabIdentifier == WorldBrowserCompositionTab)
+	{
+		FWorldBrowserModule& WorldBrowserModule = FModuleManager::LoadModuleChecked<FWorldBrowserModule>( "WorldBrowser" );
+		return SNew( SDockTab )
+			.Icon( FEditorStyle::GetBrush( "LevelEditor.Tabs.WorldBrowser" ) )
+			.Label( NSLOCTEXT("LevelEditor", "WorldBrowserCompositionTabTitle", "Levels Composition") )
+			[
+				WorldBrowserModule.CreateWorldBrowserComposition()
 			];
 	}
 	else if( TabIdentifier == TEXT("Sequencer") && FParse::Param(FCommandLine::Get(), TEXT("sequencer")) )
@@ -1143,8 +1165,21 @@ void SLevelEditor::HandleExperimentalSettingChanged(FName PropertyName)
 void SLevelEditor::RegisterWorldBrowserTabSpawner()
 {
 	TSharedPtr<FTabManager> LevelEditorTabManager = GetTabManager();
-	LevelEditorTabManager->RegisterTabSpawner( LevelEditorWorldBrowserTab, FOnSpawnTab::CreateSP<SLevelEditor, FName, FString>(this, &SLevelEditor::SpawnLevelEditorTab, LevelEditorWorldBrowserTab, FString()) )
-		.SetDisplayName(NSLOCTEXT("LevelEditorTabs", "LevelEditorWorldBrowser", "World Browser"))
+
+	LevelEditorTabManager->RegisterTabSpawner( WorldBrowserHierarchyTab, FOnSpawnTab::CreateSP<SLevelEditor, FName, FString>(this, &SLevelEditor::SpawnLevelEditorTab, WorldBrowserHierarchyTab, FString()) )
+		.SetDisplayName(NSLOCTEXT("LevelEditorTabs", "WorldBrowserHierarchy", "World Browser"))
+		.SetGroup( WorkspaceMenu::GetMenuStructure().GetLevelEditorCategory() )
+		.SetIcon( FSlateIcon(FEditorStyle::GetStyleSetName(), "LevelEditor.Tabs.WorldBrowser") );
+			
+	LevelEditorTabManager->RegisterTabSpawner( WorldBrowserDetailsTab, FOnSpawnTab::CreateSP<SLevelEditor, FName, FString>(this, &SLevelEditor::SpawnLevelEditorTab, WorldBrowserDetailsTab, FString()) )
+		.SetMenuType( ETabSpawnerMenuType::Hide )
+		.SetDisplayName(NSLOCTEXT("LevelEditorTabs", "WorldBrowserDetails", "Levels Details"))
+		.SetGroup( WorkspaceMenu::GetMenuStructure().GetLevelEditorCategory() )
+		.SetIcon( FSlateIcon(FEditorStyle::GetStyleSetName(), "LevelEditor.Tabs.WorldBrowser") );
+		
+	LevelEditorTabManager->RegisterTabSpawner( WorldBrowserCompositionTab, FOnSpawnTab::CreateSP<SLevelEditor, FName, FString>(this, &SLevelEditor::SpawnLevelEditorTab, WorldBrowserCompositionTab, FString()) )
+		.SetMenuType( ETabSpawnerMenuType::Hide )
+		.SetDisplayName(NSLOCTEXT("LevelEditorTabs", "WorldBrowserComposition", "Levels Composition"))
 		.SetGroup( WorkspaceMenu::GetMenuStructure().GetLevelEditorCategory() )
 		.SetIcon( FSlateIcon(FEditorStyle::GetStyleSetName(), "LevelEditor.Tabs.WorldBrowser") );
 }
@@ -1152,7 +1187,9 @@ void SLevelEditor::RegisterWorldBrowserTabSpawner()
 void SLevelEditor::UnregisterWorldBrowserTabSpawner()
 {
 	TSharedPtr<FTabManager> LevelEditorTabManager = GetTabManager();
-	LevelEditorTabManager->UnregisterTabSpawner( LevelEditorWorldBrowserTab );
+	LevelEditorTabManager->UnregisterTabSpawner( WorldBrowserHierarchyTab );
+	LevelEditorTabManager->UnregisterTabSpawner( WorldBrowserDetailsTab );
+	LevelEditorTabManager->UnregisterTabSpawner( WorldBrowserCompositionTab );
 }
 
 FName SLevelEditor::GetEditorModeTabId( FEditorModeID ModeID )
diff --git a/Engine/Source/Editor/LevelEditor/Public/LevelEditor.h b/Engine/Source/Editor/LevelEditor/Public/LevelEditor.h
--- a/Engine/Source/Editor/LevelEditor/Public/LevelEditor.h
+++ b/Engine/Source/Editor/LevelEditor/Public/LevelEditor.h
@@ -70,6 +70,10 @@ public:
 	 */
 	virtual void SummonLevelBrowser();
 
+	virtual void SummonWorldBrowserHierarchy();
+	virtual void SummonWorldBrowserDetails();
+	virtual void SummonWorldBrowserComposition();
+
 	// @todo remove when world-centric mode is added
 	/**
 	 * Spawns a new sequencer tab if one doesn't exist already
diff --git a/Engine/Source/Editor/UnrealEd/Private/Commandlets/CookCommandlet.cpp b/Engine/Source/Editor/UnrealEd/Private/Commandlets/CookCommandlet.cpp
--- a/Engine/Source/Editor/UnrealEd/Private/Commandlets/CookCommandlet.cpp
+++ b/Engine/Source/Editor/UnrealEd/Private/Commandlets/CookCommandlet.cpp
@@ -438,7 +438,7 @@ bool UCookCommandlet::SaveCookedPackage( UPackage* Package, uint32 SaveFlags, bo
 					if ( !World->bIsWorldInitialized)
 					{
 						// we need to initialize the world - at least need physics scene since BP construction script runs during cooking, otherwise trace won't work
-						World->InitWorld(UWorld::InitializationValues().RequiresHitProxies(false).ShouldSimulatePhysics(false).EnableTraceCollision(false).CreateNavigation(false).AllowAudioPlayback(false).CreatePhysicsScene(true).CreateWorldComposition(false));
+						World->InitWorld(UWorld::InitializationValues().RequiresHitProxies(false).ShouldSimulatePhysics(false).EnableTraceCollision(false).CreateNavigation(false).AllowAudioPlayback(false).CreatePhysicsScene(true));
 					}
 				}
 
@@ -737,12 +737,6 @@ void UCookCommandlet::CollectFilesToCook(TArray<FString>& FilesInPath)
 	GEditor->ParseMapSectionIni(*Params, MapList);
 	for (int32 MapIdx = 0; MapIdx < MapList.Num(); MapIdx++)
 	{
-		if (UWorldComposition::CollectTilesToCook(MapList[MapIdx], FilesInPath))
-		{
-			// Entry has been handled by world composition, no further processing required
-			continue;
-		}
-				
 		FilesInPath.AddUnique(MapList[MapIdx]);
 	}
 
@@ -788,13 +782,6 @@ void UCookCommandlet::CollectFilesToCook(TArray<FString>& FilesInPath)
 	{
 		FString CurrEntry = CmdLineMapEntries[CmdLineMapIdx];
 		
-		// Check if this cmd entry is related to world composition
-		if (UWorldComposition::CollectTilesToCook(CurrEntry, FilesInPath))
-		{
-			// Entry has been handled by world composition, no further processing required
-			continue;
-		}
-		
 		if (FPackageName::IsShortPackageName(CurrEntry))
 		{
 			if (FPackageName::SearchForPackageOnDisk(CurrEntry, NULL, &CurrEntry) == false)
@@ -879,35 +866,22 @@ void UCookCommandlet::CollectFilesToCook(TArray<FString>& FilesInPath)
 		FConfigCacheIni::LoadLocalIniFile(PlatformEngineIni, TEXT("Engine"), true, *Platforms[Index]->IniPlatformName());
 
 		// get the server and game default maps and cook them
-		TArray<FString> DefaultMaps;
 		FString Obj;
 		if (PlatformEngineIni.GetString(TEXT("/Script/EngineSettings.GameMapsSettings"), TEXT("GameDefaultMap"), Obj))
 		{
-			DefaultMaps.AddUnique(Obj);
+			FilesInPath.AddUnique(Obj);
 		}
 		if (PlatformEngineIni.GetString(TEXT("/Script/EngineSettings.GameMapsSettings"), TEXT("ServerDefaultMap"), Obj))
 		{
-			DefaultMaps.AddUnique(Obj);
+			FilesInPath.AddUnique(Obj);
 		}
 		if (PlatformEngineIni.GetString(TEXT("/Script/EngineSettings.GameMapsSettings"), TEXT("GlobalDefaultGameMode"), Obj))
 		{
-			DefaultMaps.AddUnique(Obj);
+			FilesInPath.AddUnique(Obj);
 		}
 		if (PlatformEngineIni.GetString(TEXT("/Script/EngineSettings.GameMapsSettings"), TEXT("GlobalDefaultServerGameMode"), Obj))
 		{
-			DefaultMaps.AddUnique(Obj);
-		}
-
-		for (FString DefaultMap : DefaultMaps)
-		{
-			// Check if this map is related to world composition
-			if (UWorldComposition::CollectTilesToCook(DefaultMap, FilesInPath))
-			{
-				// Entry has been handled by world composition, no further processing required
-				continue;
-			}
-
-			FilesInPath.AddUnique(DefaultMap);
+			FilesInPath.AddUnique(Obj);
 		}
 	}
 
@@ -1146,6 +1120,13 @@ bool UCookCommandlet::Cook(const TArray<ITargetPlatform*>& Platforms, TArray<FSt
 				{
 					World->LoadSecondaryLevels(true, &CookedPackages);
 				}
+
+				// Collect world composition tile packages to cook
+				if (World->WorldComposition)
+				{
+					World->WorldComposition->CollectTilesToCook(FilesInPath);
+				}
+
 				// maps don't compile level script actors correctly unless we do FULL GC's, they may also hold weak pointer refs that need to be reset
 				NumProcessedSinceLastGC = GCInterval; 
 
diff --git a/Engine/Source/Editor/UnrealEd/Private/EditorServer.cpp b/Engine/Source/Editor/UnrealEd/Private/EditorServer.cpp
--- a/Engine/Source/Editor/UnrealEd/Private/EditorServer.cpp
+++ b/Engine/Source/Editor/UnrealEd/Private/EditorServer.cpp
@@ -1918,11 +1918,7 @@ bool UEditorEngine::Map_Load(const TCHAR* Str, FOutputDevice& Ar)
 				// Should we display progress while loading?
 				int32 bShowProgress = 1;
 				FParse::Value(Str, TEXT("SHOWPROGRESS="), bShowProgress);
-
-				// Should we create world composition?
-				int32 bWorldComposition = 0;
-				FParse::Value(Str, TEXT("WORLDCOMPOSITION="), bWorldComposition);
-				
+			
 				UObject* OldOuter = NULL;
 
 				{
@@ -2041,7 +2037,6 @@ bool UEditorEngine::Map_Load(const TCHAR* Str, FOutputDevice& Ar)
 				Context.World()->InitWorld(UWorld::InitializationValues()
 											.ShouldSimulatePhysics(false)
 											.EnableTraceCollision(true)
-											.CreateWorldComposition(bWorldComposition != 0)
 											);
 				{
 					FBSPOps::bspValidateBrush( Context.World()->GetBrush()->Brush, 0, 1 );
diff --git a/Engine/Source/Editor/UnrealEd/Private/FileHelpers.cpp b/Engine/Source/Editor/UnrealEd/Private/FileHelpers.cpp
--- a/Engine/Source/Editor/UnrealEd/Private/FileHelpers.cpp
+++ b/Engine/Source/Editor/UnrealEd/Private/FileHelpers.cpp
@@ -1511,20 +1511,11 @@ bool FEditorFileUtils::IsValidMapFilename(const FString& MapFilename, FText& Out
 	return true;
 }
 
-/**
- * Prompts the user to save the current map if necessary, the presents a load dialog and
- * loads a new map as a world composition if selected by the user.
- */
-void FEditorFileUtils::OpenWorld()
-{
-	LoadMap(true);
-}
-
 /**
  * Prompts the user to save the current map if necessary, the presents a load dialog and
  * loads a new map if selected by the user.
  */
-void FEditorFileUtils::LoadMap(bool bWorldComposition)
+void FEditorFileUtils::LoadMap()
 {
 	if (GUnrealEd->WarnIfLightingBuildIsCurrentlyRunning())
 	{
@@ -1570,7 +1561,7 @@ void FEditorFileUtils::LoadMap(bool bWorldComposition)
 			}
 
 			FEditorDirectories::Get().SetLastDirectory(ELastDirectory::LEVEL, FPaths::GetPath(FileToOpen));
-			LoadMap( FileToOpen, false, true, bWorldComposition );
+			LoadMap( FileToOpen, false, true );
 		}
 		else
 		{
@@ -1588,7 +1579,7 @@ void FEditorFileUtils::LoadMap(bool bWorldComposition)
  * @param	LoadAsTemplate	Forces the map to load into an untitled outermost package
  *							preventing the map saving over the original file.
  */
-void FEditorFileUtils::LoadMap(const FString& InFilename, bool LoadAsTemplate, bool bShowProgress, bool bWorldComposition)
+void FEditorFileUtils::LoadMap(const FString& InFilename, bool LoadAsTemplate, bool bShowProgress)
 {
 	double LoadStartTime = FPlatformTime::Seconds();
 	
@@ -1646,7 +1637,7 @@ void FEditorFileUtils::LoadMap(const FString& InFilename, bool LoadAsTemplate, b
 		GEditorModeTools().DeactivateMode( FBuiltinEditorModes::EM_MeshPaint );
 	}
 
-	FString LoadCommand = FString::Printf( TEXT("MAP LOAD FILE=\"%s\" TEMPLATE=%d SHOWPROGRESS=%d WORLDCOMPOSITION=%d"), *Filename, LoadAsTemplate, bShowProgress, bWorldComposition );
+	FString LoadCommand = FString::Printf( TEXT("MAP LOAD FILE=\"%s\" TEMPLATE=%d SHOWPROGRESS=%d"), *Filename, LoadAsTemplate, bShowProgress );
 	bool bResult = GUnrealEd->Exec( NULL, *LoadCommand );
 
 	UWorld* World = GWorld;
@@ -2712,17 +2703,15 @@ bool FEditorFileUtils::IsFilenameValidForSaving( const FString& Filename, FText&
 
 void FEditorFileUtils::LoadDefaultMapAtStartup()
 {
-	FURL DefaultURL;
-	FURL URL( &DefaultURL, *GetDefault<UGameMapsSettings>()->EditorStartupMap, TRAVEL_Partial );
+	FString EditorStartupMap = GetDefault<UGameMapsSettings>()->EditorStartupMap;
 	
 	const bool bIncludeReadOnlyRoots = true;
-	if ( FPackageName::IsValidLongPackageName(URL.Map, bIncludeReadOnlyRoots) )
+	if ( FPackageName::IsValidLongPackageName(EditorStartupMap, bIncludeReadOnlyRoots) )
 	{
-		FString MapFilenameToLoad = FPackageName::LongPackageNameToFilename( URL.Map );
-		bool bWorldComposition = URL.HasOption(TEXT("worldcomposition"));
+		FString MapFilenameToLoad = FPackageName::LongPackageNameToFilename( EditorStartupMap );
 
 		bIsLoadingDefaultStartupMap = true;
-		FEditorFileUtils::LoadMap( MapFilenameToLoad + FPackageName::GetMapPackageExtension(), GUnrealEd->IsTemplateMap(URL.Map), true, bWorldComposition );
+		FEditorFileUtils::LoadMap( MapFilenameToLoad + FPackageName::GetMapPackageExtension(), GUnrealEd->IsTemplateMap(EditorStartupMap), true );
 		bIsLoadingDefaultStartupMap = false;
 	}
 }
diff --git a/Engine/Source/Editor/UnrealEd/Private/PlayLevel.cpp b/Engine/Source/Editor/UnrealEd/Private/PlayLevel.cpp
--- a/Engine/Source/Editor/UnrealEd/Private/PlayLevel.cpp
+++ b/Engine/Source/Editor/UnrealEd/Private/PlayLevel.cpp
@@ -946,7 +946,6 @@ void UEditorEngine::PlayOnLocalPc(FString MapNameOverride, FString URLParms, FSt
 		{
 			// Open world composition from original folder
 			FString MapName = EditorContext.World()->GetOutermost()->GetName();
-			MapName += TEXT("?worldcomposition");
 			SavedMapNames.Add(MapName);
 		}
 		else
@@ -1187,8 +1186,6 @@ void UEditorEngine::PlayUsingLauncher()
 		{
 			// Open world composition from original folder
 			FString MapName = EditorContext.World()->GetOutermost()->GetName();
-			MapName += TEXT("?worldcomposition");
-
 			MapNames.Add(MapName);
 		}
 		else
diff --git a/Engine/Source/Editor/UnrealEd/Private/UnrealEdMisc.cpp b/Engine/Source/Editor/UnrealEd/Private/UnrealEdMisc.cpp
--- a/Engine/Source/Editor/UnrealEd/Private/UnrealEdMisc.cpp
+++ b/Engine/Source/Editor/UnrealEd/Private/UnrealEdMisc.cpp
@@ -164,10 +164,9 @@ void FUnrealEdMisc::OnInit()
 				// be reflected in the splash screen status
 				const bool bShowProgress = false;
 				const bool bLoadAsTemplate = false;
-				const bool bWorldComposition = URL.HasOption(TEXT("worldcomposition"));
 
 				// Load the map
-				FEditorFileUtils::LoadMap(InitialMapName, bLoadAsTemplate, bShowProgress, bWorldComposition);
+				FEditorFileUtils::LoadMap(InitialMapName, bLoadAsTemplate, bShowProgress);
 				bMapLoaded = true;
 			}
 		}
diff --git a/Engine/Source/Editor/UnrealEd/Public/FileHelpers.h b/Engine/Source/Editor/UnrealEd/Public/FileHelpers.h
--- a/Engine/Source/Editor/UnrealEd/Public/FileHelpers.h
+++ b/Engine/Source/Editor/UnrealEd/Public/FileHelpers.h
@@ -43,19 +43,6 @@ public:
 	 */
 	static void ResetLevelFilenames();	
 
-	
-	////////////////////////////////////////////////////////////////////////////
-	// World
-
-	/**
-	 * Prompts the user to save the current map if necessary, the presents a load dialog and
-	 * loads a new map if selected by the user.
-	 */
-	static UNREALED_API void OpenWorld();
-
-	////////////////////////////////////////////////////////////////////////////
-
-
 	////////////////////////////////////////////////////////////////////////////
 	// Loading
 
@@ -68,10 +55,8 @@ public:
 	/**
 	 * Prompts the user to save the current map if necessary, the presents a load dialog and
 	 * loads a new map if selected by the user.
-	 * 
-	 * @param	bWorldComposition	Whether we should load new map as persistent level for world composition
 	 */
-	static UNREALED_API void LoadMap(bool bWorldComposition = false);
+	static UNREALED_API void LoadMap();
 
 	/**
 	 * Loads the specified map.  Does not prompt the user to save the current map.
@@ -81,10 +66,8 @@ public:
 	 * @param	LoadAsTemplate	Forces the map to load into an untitled outermost package
 	 *							preventing the map saving over the original file.
 	 * @param	bShowProgress	Whether to show a progress dialog as the map loads
-	 *
-	 * @param	bWorldComposition	Whether to create a world composition by reading level packages from this map folder
 	 */
-	static UNREALED_API void LoadMap(const FString& Filename, bool LoadAsTemplate = false, const bool bShowProgress=true, bool bWorldComposition = false);
+	static UNREALED_API void LoadMap(const FString& Filename, bool LoadAsTemplate = false, const bool bShowProgress=true);
 
 	////////////////////////////////////////////////////////////////////////////
 	// Saving
diff --git a/Engine/Source/Editor/WorldBrowser/Private/LevelCollectionModel.cpp b/Engine/Source/Editor/WorldBrowser/Private/LevelCollectionModel.cpp
--- a/Engine/Source/Editor/WorldBrowser/Private/LevelCollectionModel.cpp
+++ b/Engine/Source/Editor/WorldBrowser/Private/LevelCollectionModel.cpp
@@ -26,6 +26,8 @@ FLevelCollectionModel::FLevelCollectionModel(const TWeakObjectPtr<UEditorEngine>
 
 FLevelCollectionModel::~FLevelCollectionModel()
 {
+	SaveSettings();
+	
 	Filters->OnChanged().RemoveAll(this);
 	FWorldDelegates::LevelAddedToWorld.RemoveAll(this);
 	FWorldDelegates::LevelRemovedFromWorld.RemoveAll(this);
@@ -36,6 +38,8 @@ FLevelCollectionModel::~FLevelCollectionModel()
 
 void FLevelCollectionModel::Initialize()
 {
+	LoadSettings();
+	
 	CurrentWorld = Editor->GetEditorWorldContext().World();
 
 	Filters->OnChanged().AddSP(this, &FLevelCollectionModel::OnFilterChanged);
@@ -305,6 +309,11 @@ void FLevelCollectionModel::RemoveFilter(const TSharedRef<LevelFilter>& InFilter
 	OnFilterChanged();
 }
 
+bool FLevelCollectionModel::IsFilterActive() const
+{
+	return (AllLevelsList.Num() != FilteredLevelsList.Num());
+}
+
 void FLevelCollectionModel::SetSelectedLevels(const FLevelModelList& InList)
 {
 	// Clear selection flag from currently selected levels
@@ -616,7 +625,7 @@ TSharedPtr<FLevelDragDropOp> FLevelCollectionModel::CreateDragDropOp() const
 
 bool FLevelCollectionModel::PassesAllFilters(TSharedPtr<FLevelModel> Item) const
 {
-	if (Filters->PassesAllFilters(Item))
+	if (Item->IsPersistent() || Filters->PassesAllFilters(Item))
 	{
 		return true;
 	}
@@ -913,6 +922,21 @@ void FLevelCollectionModel::SetDisplayPathsState(bool InDisplayPaths)
 	}
 }
 
+bool FLevelCollectionModel::GetDisplayActorsCountState() const
+{
+	return bDisplayActorsCount;
+}
+
+void FLevelCollectionModel::SetDisplayActorsCountState(bool InDisplayActorsCount)
+{
+	bDisplayActorsCount = InDisplayActorsCount;
+
+	for (auto It = AllLevelsList.CreateIterator(); It; ++It)
+	{
+		(*It)->UpdateDisplayName();
+	}
+}
+
 void FLevelCollectionModel::BroadcastSelectionChanged()
 {
 	SelectionChanged.Broadcast();
@@ -1266,6 +1290,30 @@ const TSharedRef<const FUICommandList> FLevelCollectionModel::GetCommandList() c
 	return CommandList;
 }
 
+const FString ConfigIniSection = TEXT("WorldBrowser");
+
+void FLevelCollectionModel::LoadSettings()
+{
+	// Display paths
+	bool bDisplayPathsSetting = false;
+	GConfig->GetBool(*ConfigIniSection, TEXT("DisplayPaths"), bDisplayPathsSetting, GEditorUserSettingsIni);
+	SetDisplayPathsState(bDisplayPathsSetting);
+
+	// Display actors count
+	bool bDisplayActorsCountSetting = false;
+	GConfig->GetBool(*ConfigIniSection, TEXT("DisplayActorsCount"), bDisplayActorsCountSetting, GEditorUserSettingsIni);
+	SetDisplayActorsCountState(bDisplayActorsCountSetting);
+}
+	
+void FLevelCollectionModel::SaveSettings()
+{
+	// Display paths
+	GConfig->SetBool(*ConfigIniSection, TEXT("DisplayPaths"), GetDisplayPathsState(), GEditorUserSettingsIni);
+
+	// Display actors count
+	GConfig->SetBool(*ConfigIniSection, TEXT("DisplayActorsCount"), GetDisplayActorsCountState(), GEditorUserSettingsIni);
+}
+
 void FLevelCollectionModel::RefreshBrowser_Executed()
 {
 	PopulateLevelsList();
diff --git a/Engine/Source/Editor/WorldBrowser/Private/LevelCollectionModel.h b/Engine/Source/Editor/WorldBrowser/Private/LevelCollectionModel.h
--- a/Engine/Source/Editor/WorldBrowser/Private/LevelCollectionModel.h
+++ b/Engine/Source/Editor/WorldBrowser/Private/LevelCollectionModel.h
@@ -63,6 +63,9 @@ public:
 
 	/** Removes a filter which restricted the Levels shown in UI */
 	void RemoveFilter(const TSharedRef<LevelFilter>& InFilter);
+
+	/**	@return	Whether level filtering is active now */
+	bool IsFilterActive() const;
 	
 	/**	Iterates through level hierarchy with given Visitor */
 	void IterateHierarchy(FLevelModelVisitor& Visitor);
@@ -136,8 +139,8 @@ public:
 	/**	Unregisters level details customizations */
 	virtual void UnregisterDetailsCustomization(class FPropertyEditorModule& PropertyModule, TSharedPtr<class IDetailsView> InDetailsView);
 
-	/** @return	Whether this level collection model supports grid view */
-	virtual bool SupportsGridView() const { return true; };
+	/** @return	Whether this level collection model is a tile world */
+	virtual bool IsTileWorld() const { return false; };
 
 	/** Rebuilds levels collection */
 	void PopulateLevelsList();
@@ -208,6 +211,12 @@ public:
 	/** Sets 'display paths', whether to show long package name in level display name */
 	void SetDisplayPathsState(bool bDisplayPaths);
 
+	/** @return	whether 'display actors count' is enabled */
+	bool GetDisplayActorsCountState() const;
+
+	/** Sets 'display actors count', whether to show actors count next to level name */
+	void SetDisplayActorsCountState(bool bDisplayActorsCount);
+
 	/**	Broadcasts whenever items selection has changed */
 	FSimpleEvent SelectionChanged;
 	void BroadcastSelectionChanged();
@@ -261,6 +270,12 @@ public:
 	/** @return	The UICommandList supported by this collection */
 	const TSharedRef<const FUICommandList> GetCommandList() const;
 
+	/**  */
+	void LoadSettings();
+	
+	/**  */
+	void SaveSettings();
+
 protected:
 	/** Refreshes current cached data */
 	void RefreshBrowser_Executed();
@@ -474,6 +489,9 @@ protected:
 	/** Whether we should show long package names in level display names */
 	bool								bDisplayPaths;
 
+	/** Whether we should show actors count next to level name */
+	bool								bDisplayActorsCount;
+
 	/** true if the SCC Check-Out option is available */
 	mutable bool						bCanExecuteSCCCheckOut;
 
diff --git a/Engine/Source/Editor/WorldBrowser/Private/LevelModel.cpp b/Engine/Source/Editor/WorldBrowser/Private/LevelModel.cpp
--- a/Engine/Source/Editor/WorldBrowser/Private/LevelModel.cpp
+++ b/Engine/Source/Editor/WorldBrowser/Private/LevelModel.cpp
@@ -668,7 +668,7 @@ void FLevelModel::UpdateDisplayName()
 	if (HasValidPackage())
 	{
 		// Append actors count
-		if (IsLoaded())
+		if (LevelCollectionModel.GetDisplayActorsCountState() && IsLoaded())
 		{
 			DisplayName += TEXT(" (");
 			DisplayName.AppendInt(LevelActorsCount);
diff --git a/Engine/Source/Editor/WorldBrowser/Private/SWorldDetails.cpp b/Engine/Source/Editor/WorldBrowser/Private/SWorldDetails.cpp
new file mode 100644
--- /dev/null
+++ b/Engine/Source/Editor/WorldBrowser/Private/SWorldDetails.cpp
@@ -0,0 +1,237 @@
+// Copyright 1998-2014 Epic Games, Inc. All Rights Reserved.
+#include "WorldBrowserPrivatePCH.h"
+
+#include "Editor/PropertyEditor/Public/IDetailsView.h"
+#include "Editor/LevelEditor/Public/LevelEditor.h"
+#include "SWorldDetails.h"
+
+#define LOCTEXT_NAMESPACE "WorldBrowser"
+
+SWorldDetails::SWorldDetails()
+{
+}
+
+SWorldDetails::~SWorldDetails()
+{
+	GEditorModeTools().DeactivateMode(FBuiltinEditorModes::EM_StreamingLevel);
+	
+	FWorldBrowserModule& WorldBrowserModule = FModuleManager::GetModuleChecked<FWorldBrowserModule>("WorldBrowser");
+	WorldBrowserModule.OnBrowseWorld.RemoveAll(this);
+	
+	OnBrowseWorld(nullptr);
+}
+
+void SWorldDetails::Construct(const FArguments& InArgs)
+{
+	FWorldBrowserModule& WorldBrowserModule = FModuleManager::GetModuleChecked<FWorldBrowserModule>("WorldBrowser");
+	WorldBrowserModule.OnBrowseWorld.AddSP(this, &SWorldDetails::OnBrowseWorld);
+	
+	OnBrowseWorld(InArgs._InWorld);
+}
+
+void SWorldDetails::OnBrowseWorld(UWorld* InWorld)
+{
+	// Remove old world bindings
+	ChildSlot
+	[
+		SNullWidget::NullWidget
+	];
+	
+	if (WorldModel.IsValid())
+	{
+		FPropertyEditorModule& PropertyModule = FModuleManager::GetModuleChecked<FPropertyEditorModule>("PropertyEditor");
+		WorldModel->UnregisterDetailsCustomization(PropertyModule, DetailsView);
+		WorldModel->SelectionChanged.RemoveAll(this);
+	}
+	
+	WorldModel = nullptr;
+
+	// Bind to a new world
+	if (InWorld)
+	{
+		FWorldBrowserModule& WorldBrowserModule = FModuleManager::GetModuleChecked<FWorldBrowserModule>("WorldBrowser");
+		WorldModel = WorldBrowserModule.SharedWorldModel(InWorld);
+
+		WorldModel->SelectionChanged.AddSP(this, &SWorldDetails::OnSelectionChanged);
+	
+		FPropertyEditorModule& PropertyModule = FModuleManager::LoadModuleChecked<FPropertyEditorModule>("PropertyEditor");
+		FDetailsViewArgs Args(false, false, false, false, true);
+		Args.bShowActorLabel = false;
+	
+		DetailsView = PropertyModule.CreateDetailView(Args);
+		ChildSlot
+		[
+			SNew(SVerticalBox)
+
+			// Inspect level box
+			+SVerticalBox::Slot()
+			.AutoHeight()
+			[
+				SNew(SBorder)
+				.BorderImage(FEditorStyle::GetBrush(TEXT("ToolPanel.GroupBorder")))
+				[
+					SNew(SHorizontalBox)
+
+					+SHorizontalBox::Slot()
+					.AutoWidth()
+					.VAlign(VAlign_Center)
+					.HAlign(HAlign_Left)
+					[
+						SNew(STextBlock)
+						.Text(FText::FromString(TEXT("Inspect level:")))
+					]
+					
+					+SHorizontalBox::Slot()
+					.FillWidth(1.f)
+					.VAlign(VAlign_Center)
+					.HAlign(HAlign_Left)
+					.Padding(4,0,4,0)
+					[
+						SAssignNew(SubLevelsComboBox, SComboBox<TSharedPtr<FLevelModel>>)
+						.OptionsSource(&WorldModel->GetFilteredLevels())
+						.OnSelectionChanged(this, &SWorldDetails::OnSetInspectedLevel)
+						.OnGenerateWidget(this, &SWorldDetails::HandleInspectedLevelComboBoxGenerateWidget)
+						.Content()
+						[
+							SNew(STextBlock)
+							.Text(this, &SWorldDetails::GetInspectedLevelText)
+						]
+					]
+					
+					// Button to summon levels hierarchy tab
+					+SHorizontalBox::Slot()
+					.AutoWidth()
+					.VAlign(VAlign_Center)
+					.HAlign(HAlign_Right)
+					[
+						SNew(SButton)
+						.ButtonStyle(FEditorStyle::Get(), "ToggleButton")
+						.OnClicked(this, &SWorldDetails::OnSummonHierarchy)
+						.ToolTipText(LOCTEXT("SummonHierarchyToolTipText", "Summons sublevels hierarchy"))
+						.HAlign(HAlign_Center)
+						.VAlign(VAlign_Center)
+						.Content()
+						[
+							SNew(SImage)
+							.Image(this, &SWorldDetails::GetSummonHierarchyBrush)
+						]
+					]
+
+					// Button to summon world composition tab
+					+SHorizontalBox::Slot()
+					.AutoWidth()
+					.VAlign(VAlign_Center)
+					.HAlign(HAlign_Right)
+					[
+						SNew(SButton)
+						.Visibility(this, &SWorldDetails::GetCompositionButtonVisibility)
+						.ButtonStyle(FEditorStyle::Get(), "ToggleButton")
+						.OnClicked(this, &SWorldDetails::OnSummonComposition)
+						.ToolTipText(LOCTEXT("SummonHierarchyToolTipText", "Summons world composition"))
+						.HAlign(HAlign_Center)
+						.VAlign(VAlign_Center)
+						.Content()
+						[
+							SNew(SImage)
+							.Image(this, &SWorldDetails::GetSummonCompositionBrush)
+						]
+					]
+				]
+			]
+			
+			// Level details
+			+SVerticalBox::Slot()
+			.FillHeight(1.f)
+			.Padding(0,4,0,0)
+			[
+				SNew(SBorder)
+				.BorderImage(FEditorStyle::GetBrush(TEXT("ToolPanel.GroupBorder")))
+				[
+					DetailsView.ToSharedRef()
+				]
+			]
+		];
+
+		WorldModel->RegisterDetailsCustomization(PropertyModule, DetailsView);
+	}
+}
+
+void SWorldDetails::OnSelectionChanged()
+{
+	auto SelectedLevels = WorldModel->GetSelectedLevels();
+	TArray<UObject*> TileProperties;
+	
+	for (auto It = SelectedLevels.CreateIterator(); It; ++It)
+	{
+		UObject* PropertiesObject = (*It)->GetNodeObject();
+		if (PropertiesObject)
+		{
+			TileProperties.Add(PropertiesObject);
+		}
+	}
+
+	DetailsView->SetObjects(TileProperties, true);
+}
+
+void SWorldDetails::OnSetInspectedLevel(TSharedPtr<FLevelModel> InLevelModel, ESelectInfo::Type SelectInfo)
+{
+	FLevelModelList SelectedLevels; SelectedLevels.Add(InLevelModel);
+	WorldModel->SetSelectedLevels(SelectedLevels);
+}
+
+TSharedRef<SWidget> SWorldDetails::HandleInspectedLevelComboBoxGenerateWidget(TSharedPtr<FLevelModel> InLevelModel) const
+{
+	return SNew(SBox)
+	.Padding(4)
+	[
+		SNew(STextBlock).Text(InLevelModel->GetDisplayName())
+	];
+}
+
+FText SWorldDetails::GetInspectedLevelText() const
+{
+	const FLevelModelList& SelectedLevels = WorldModel->GetSelectedLevels();
+	
+	if (SelectedLevels.Num() > 1)
+	{
+		return LOCTEXT("MultipleInspectedLevelText", "Multiple Values");
+	}
+	else if (SelectedLevels.Num() == 1)
+	{
+		return FText::FromString(SelectedLevels[0]->GetDisplayName());
+	}
+
+	return LOCTEXT("EmptyInspectedLevelText", "None");
+}
+
+FReply SWorldDetails::OnSummonHierarchy()
+{
+	FLevelEditorModule& LevelEditorModule = FModuleManager::LoadModuleChecked<FLevelEditorModule>( "LevelEditor" );
+	LevelEditorModule.SummonWorldBrowserHierarchy();
+	return FReply::Handled();
+}
+
+const FSlateBrush* SWorldDetails::GetSummonHierarchyBrush() const
+{
+	return FEditorStyle::GetBrush("WorldBrowser.HierarchyButtonBrush");
+}
+
+EVisibility SWorldDetails::GetCompositionButtonVisibility() const
+{
+	return WorldModel->IsTileWorld() ? EVisibility::Visible : EVisibility::Collapsed;
+}
+
+FReply SWorldDetails::OnSummonComposition()
+{
+	FLevelEditorModule& LevelEditorModule = FModuleManager::LoadModuleChecked<FLevelEditorModule>( "LevelEditor" );
+	LevelEditorModule.SummonWorldBrowserComposition();
+	return FReply::Handled();
+}
+
+const FSlateBrush* SWorldDetails::GetSummonCompositionBrush() const
+{
+	return FEditorStyle::GetBrush("WorldBrowser.CompositionButtonBrush");
+}
+
+
+#undef LOCTEXT_NAMESPACE
\ No newline at end of file
diff --git a/Engine/Source/Editor/WorldBrowser/Private/SWorldDetails.h b/Engine/Source/Editor/WorldBrowser/Private/SWorldDetails.h
new file mode 100644
--- /dev/null
+++ b/Engine/Source/Editor/WorldBrowser/Private/SWorldDetails.h
@@ -0,0 +1,50 @@
+// Copyright 1998-2014 Epic Games, Inc. All Rights Reserved.
+#pragma once
+
+#include "Editor/PropertyEditor/Public/PropertyEditorModule.h"
+
+class FLevelCollectionModel;
+//----------------------------------------------------------------
+//
+//
+//----------------------------------------------------------------
+class SWorldDetails
+	: public SCompoundWidget
+{
+public:
+	SLATE_BEGIN_ARGS(SWorldDetails)
+		:_InWorld(nullptr)
+		{}
+		SLATE_ARGUMENT(UWorld*, InWorld)
+	SLATE_END_ARGS()
+
+	void Construct(const FArguments& InArgs);
+	SWorldDetails();
+	~SWorldDetails();
+
+private:
+	/**  */
+	void OnBrowseWorld(UWorld* InWorld);
+
+	/** Handles selection changes in data source */
+	void OnSelectionChanged();
+
+	/**  */
+	void OnSetInspectedLevel(TSharedPtr<FLevelModel> InLevelModel, ESelectInfo::Type SelectInfo);
+	TSharedRef<SWidget> HandleInspectedLevelComboBoxGenerateWidget(TSharedPtr<FLevelModel> InLevelModel) const;
+	FText GetInspectedLevelText() const;
+	
+	/**  */
+	FReply OnSummonHierarchy();
+	const FSlateBrush* GetSummonHierarchyBrush() const;
+
+	/**  */
+	EVisibility GetCompositionButtonVisibility() const;
+	FReply OnSummonComposition();
+	const FSlateBrush* GetSummonCompositionBrush() const;
+		
+private:
+	TSharedPtr<FLevelCollectionModel>				WorldModel;
+	TSharedPtr<IDetailsView>						DetailsView;
+	TSharedPtr<SComboBox<TSharedPtr<FLevelModel>>>	SubLevelsComboBox;
+};
diff --git a/Engine/Source/Editor/WorldBrowser/Private/SWorldDetailsView.cpp b/Engine/Source/Editor/WorldBrowser/Private/SWorldDetailsView.cpp
deleted file mode 100644
diff --git a/Engine/Source/Editor/WorldBrowser/Private/SWorldDetailsView.h b/Engine/Source/Editor/WorldBrowser/Private/SWorldDetailsView.h
deleted file mode 100644
diff --git a/Engine/Source/Editor/WorldBrowser/Private/SWorldGridItem.cpp b/Engine/Source/Editor/WorldBrowser/Private/SWorldGridItem.cpp
deleted file mode 100644
diff --git a/Engine/Source/Editor/WorldBrowser/Private/SWorldGridItem.h b/Engine/Source/Editor/WorldBrowser/Private/SWorldGridItem.h
deleted file mode 100644
diff --git a/Engine/Source/Editor/WorldBrowser/Private/SWorldGridView.cpp b/Engine/Source/Editor/WorldBrowser/Private/SWorldGridView.cpp
deleted file mode 100644
diff --git a/Engine/Source/Editor/WorldBrowser/Private/SWorldGridView.h b/Engine/Source/Editor/WorldBrowser/Private/SWorldGridView.h
deleted file mode 100644
diff --git a/Engine/Source/Editor/WorldBrowser/Private/SWorldHierarchy.cpp b/Engine/Source/Editor/WorldBrowser/Private/SWorldHierarchy.cpp
new file mode 100644
--- /dev/null
+++ b/Engine/Source/Editor/WorldBrowser/Private/SWorldHierarchy.cpp
@@ -0,0 +1,680 @@
+// Copyright 1998-2014 Epic Games, Inc. All Rights Reserved.
+#include "WorldBrowserPrivatePCH.h"
+
+#include "Editor/LevelEditor/Public/LevelEditor.h"
+#include "TextFilter.h"
+#include "SWorldHierarchy.h"
+#include "SWorldHierarchyItem.h"
+
+#define LOCTEXT_NAMESPACE "WorldBrowser"
+
+typedef STreeView<TSharedPtr<FLevelModel>> SLevelsTreeWidget;
+typedef TTextFilter<const TSharedPtr<FLevelModel>& > LevelTextFilter;
+
+//----------------------------------------------------------------
+//
+//
+//----------------------------------------------------------------
+class SWorldHierarchyImpl 
+	: public SCompoundWidget
+{
+public:
+	SLATE_BEGIN_ARGS(SWorldHierarchyImpl) {}
+		SLATE_ARGUMENT(TSharedPtr<FLevelCollectionModel>, InWorldModel)
+	SLATE_END_ARGS()
+
+	SWorldHierarchyImpl()
+		: bUpdatingSelection(false)
+	{
+	}
+
+	~SWorldHierarchyImpl()
+	{
+		WorldModel->SelectionChanged.RemoveAll(this);
+		WorldModel->HierarchyChanged.RemoveAll(this);
+		WorldModel->CollectionChanged.RemoveAll(this);
+	}
+	
+	void Construct(const FArguments& InArgs)
+	{
+		WorldModel = InArgs._InWorldModel;
+		check(WorldModel.IsValid());
+	
+		WorldModel->SelectionChanged.AddSP(this, &SWorldHierarchyImpl::OnUpdateSelection);
+		WorldModel->HierarchyChanged.AddSP(this, &SWorldHierarchyImpl::RefreshView);
+		WorldModel->CollectionChanged.AddSP(this, &SWorldHierarchyImpl::RefreshView);
+		
+		SearchBoxLevelFilter = MakeShareable(new LevelTextFilter(
+				LevelTextFilter::FItemToStringArray::CreateSP(this, &SWorldHierarchyImpl::TransformLevelToString)
+		));
+
+		HeaderRowWidget =
+			SNew( SHeaderRow )
+			.Visibility(EVisibility::Collapsed)
+
+			/** Level visibility column */
+			+ SHeaderRow::Column(HierarchyColumns::ColumnID_Visibility)
+			.FixedWidth(24.0f)
+			.HeaderContent()
+			[
+				SNew(STextBlock)
+				.ToolTipText(NSLOCTEXT("WorldBrowser", "Visibility", "Visibility"))
+			]
+
+			/** LevelName label column */
+			+ SHeaderRow::Column( HierarchyColumns::ColumnID_LevelLabel )
+				.FillWidth( 0.45f )
+				.HeaderContent()
+				[
+					SNew(STextBlock)
+						.ToolTipText(LOCTEXT("Column_LevelNameLabel", "Level"))
+				]
+
+			/** Level lock column */
+			+ SHeaderRow::Column( HierarchyColumns::ColumnID_Lock )
+				.FixedWidth( 24.0f )
+				.HeaderContent()
+				[
+					SNew(STextBlock)
+						.ToolTipText(NSLOCTEXT("WorldBrowser", "Lock", "Lock"))
+				]
+	
+			/** Level kismet column */
+			+ SHeaderRow::Column( HierarchyColumns::ColumnID_Kismet )
+				.FixedWidth( 24.0f )
+				.HeaderContent()
+				[
+					SNew(STextBlock)
+						.ToolTipText(NSLOCTEXT("WorldBrowser", "Blueprint", "Open the level blueprint for this Level"))
+				]
+
+			/** Level SCC status column */
+			+ SHeaderRow::Column( HierarchyColumns::ColumnID_SCCStatus )
+				.FixedWidth( 24.0f )
+				.HeaderContent()
+				[
+					SNew(STextBlock)
+						.ToolTipText(NSLOCTEXT("WorldBrowser", "SCCStatus", "Status in Source Control"))
+				]
+
+			/** Level save column */
+			+ SHeaderRow::Column( HierarchyColumns::ColumnID_Save )
+				.FixedWidth( 24.0f )
+				.HeaderContent()
+				[
+					SNew(STextBlock)
+						.ToolTipText(NSLOCTEXT("WorldBrowser", "Save", "Save this Level"))
+				];
+
+
+		ChildSlot
+		[
+			SNew(SVerticalBox)
+			// Filter box
+			+SVerticalBox::Slot()
+			.AutoHeight()
+			[
+				SNew( SSearchBox )
+				.ToolTipText(LOCTEXT("FilterSearchToolTip", "Type here to search Levels").ToString())
+				.HintText(LOCTEXT("FilterSearchHint", "Search Levels"))
+				.OnTextChanged(SearchBoxLevelFilter.Get(), &LevelTextFilter::SetRawFilterText)
+			]
+			// Hierarchy
+			+SVerticalBox::Slot()
+			.FillHeight(1.f)
+			[
+				SAssignNew(TreeWidget, SLevelsTreeWidget)
+				.TreeItemsSource(&WorldModel->GetRootLevelList())
+				.SelectionMode(ESelectionMode::Multi)
+				.OnGenerateRow(this, &SWorldHierarchyImpl::GenerateTreeRow)
+				.OnGetChildren( this, &SWorldHierarchyImpl::GetChildrenForTree)
+				.OnSelectionChanged(this, &SWorldHierarchyImpl::OnSelectionChanged)
+				.OnExpansionChanged(this, &SWorldHierarchyImpl::OnExpansionChanged)
+				.OnMouseButtonDoubleClick(this, &SWorldHierarchyImpl::OnTreeViewMouseButtonDoubleClick)
+				.OnContextMenuOpening(this, &SWorldHierarchyImpl::ConstructLevelContextMenu)
+				.HeaderRow(HeaderRowWidget.ToSharedRef())
+			]
+
+			// Separator
+			+SVerticalBox::Slot()
+			.AutoHeight()
+			.Padding(0, 0, 0, 1)
+			[
+				SNew(SSeparator)
+			]
+		
+			// View options
+			+SVerticalBox::Slot()
+			.AutoHeight()
+			[
+				SNew(SHorizontalBox)
+
+				// Asset count
+				+SHorizontalBox::Slot()
+				.FillWidth(1.f)
+				.VAlign(VAlign_Center)
+				.Padding(8, 0)
+				[
+					SNew( STextBlock )
+					.Text( this, &SWorldHierarchyImpl::GetFilterStatusText )
+					.ColorAndOpacity( this, &SWorldHierarchyImpl::GetFilterStatusTextColor )
+				]
+
+				// View mode combo button
+				+SHorizontalBox::Slot()
+				.AutoWidth()
+				[
+					SAssignNew( ViewOptionsComboButton, SComboButton )
+					.ContentPadding(0)
+					.ForegroundColor( this, &SWorldHierarchyImpl::GetViewButtonForegroundColor )
+					.ButtonStyle( FEditorStyle::Get(), "ToggleButton" ) // Use the tool bar item style for this button
+					.OnGetMenuContent( this, &SWorldHierarchyImpl::GetViewButtonContent )
+					.ButtonContent()
+					[
+						SNew(SHorizontalBox)
+
+						+SHorizontalBox::Slot()
+						.AutoWidth()
+						.VAlign(VAlign_Center)
+						[
+							SNew(SImage).Image( FEditorStyle::GetBrush("GenericViewButton") )
+						]
+
+						+SHorizontalBox::Slot()
+						.AutoWidth()
+						.Padding(2, 0, 0, 0)
+						.VAlign(VAlign_Center)
+						[
+							SNew(STextBlock).Text( LOCTEXT("ViewButton", "View Options") )
+						]
+					]
+				]
+			]
+		];
+	
+		WorldModel->AddFilter(SearchBoxLevelFilter.ToSharedRef());
+		RefreshView();
+	}
+
+	/** Regenerates current items */
+	void RefreshView()
+	{
+		TreeWidget->RequestTreeRefresh();
+	
+		// Sync items expansion state
+		struct FExpander : public FLevelModelVisitor 
+		{
+			TSharedPtr<SLevelsTreeWidget> TreeWidget;
+			virtual void Visit(FLevelModel& Item) OVERRIDE
+			{
+				TreeWidget->SetItemExpansion(Item.AsShared(), Item.GetLevelExpansionFlag());
+			};
+		} Expander;
+		Expander.TreeWidget = TreeWidget;
+		
+		//Apply expansion
+		WorldModel->IterateHierarchy(Expander);
+	}
+
+private:
+	/** Creates an item for the tree view */
+	TSharedRef<ITableRow> GenerateTreeRow(TSharedPtr<FLevelModel> Item, const TSharedRef<STableViewBase>& OwnerTable)
+	{
+		check(Item.IsValid());
+
+		return SNew(SWorldHierarchyItem, OwnerTable)
+				.InWorldModel(WorldModel)
+				.InItemModel(Item)
+				.IsItemExpanded(this, &SWorldHierarchyImpl::IsTreeItemExpanded, Item)
+				.HighlightText(this, &SWorldHierarchyImpl::GetSearchBoxText);
+
+	}
+
+	/** Handler for returning a list of children associated with a particular tree node */
+	void GetChildrenForTree(TSharedPtr<FLevelModel> Item, FLevelModelList& OutChildren)
+	{
+		OutChildren = Item->GetChildren();
+	}
+
+	/** @return the SWidget containing the context menu */
+	TSharedPtr<SWidget> ConstructLevelContextMenu() const
+	{
+		if (!WorldModel->IsReadOnly())
+		{
+			FMenuBuilder MenuBuilder(true, WorldModel->GetCommandList());
+			WorldModel->BuildHierarchyMenu(MenuBuilder);
+			return MenuBuilder.MakeWidget();
+		}
+
+		return SNullWidget::NullWidget;	
+	}
+
+	/** Called by TreeView widget whenever tree item expanded or collapsed */
+	void OnExpansionChanged(TSharedPtr<FLevelModel> Item, bool bIsItemExpanded)
+	{
+		Item->SetLevelExpansionFlag(bIsItemExpanded);
+	}
+
+	/** Called by TreeView widget whenever selection is changed */
+	void OnSelectionChanged(const TSharedPtr<FLevelModel> Item, ESelectInfo::Type SelectInfo)
+	{
+		if (bUpdatingSelection)
+		{
+			return;
+		}
+
+		bUpdatingSelection = true;
+		WorldModel->SetSelectedLevels(TreeWidget->GetSelectedItems());
+		bUpdatingSelection = false;
+	}
+
+	/** Handles selection changes in data source */
+	void OnUpdateSelection()
+	{
+		if (bUpdatingSelection)
+		{
+			return;
+		}
+
+		bUpdatingSelection = true;
+		const auto& SelectedItems = WorldModel->GetSelectedLevels();
+		TreeWidget->ClearSelection();
+		for (auto It = SelectedItems.CreateConstIterator(); It; ++It)
+		{
+			TreeWidget->SetItemSelection(*It, true);
+		}
+
+		if (SelectedItems.Num() == 1)
+		{
+			TreeWidget->RequestScrollIntoView(SelectedItems[0]);
+		}
+	
+		RefreshView();
+
+		bUpdatingSelection = false;
+	}
+
+	/** 
+	 *	Called by STreeView when the user double-clicks on an item 
+	 *
+	 *	@param	Item	The item that was double clicked
+	 */
+	void OnTreeViewMouseButtonDoubleClick(TSharedPtr<FLevelModel> Item)
+	{
+		Item->MakeLevelCurrent();
+	}
+
+	/**
+	 * Checks to see if this widget supports keyboard focus.  Override this in derived classes.
+	 *
+	 * @return  True if this widget can take keyboard focus
+	 */
+	virtual bool SupportsKeyboardFocus() const OVERRIDE
+	{
+		return true;
+	}
+
+	/**
+	 * Called after a key is pressed when this widget has keyboard focus
+	 *
+	 * @param MyGeometry The Geometry of the widget receiving the event
+	 * @param  InKeyboardEvent  Keyboard event
+	 *
+	 * @return  Returns whether the event was handled, along with other possible actions
+	 */
+	virtual FReply OnKeyDown(const FGeometry& MyGeometry, const FKeyboardEvent& InKeyboardEvent) OVERRIDE
+	{
+		if (WorldModel->GetCommandList()->ProcessCommandBindings(InKeyboardEvent))
+		{
+			return FReply::Handled();
+		}
+	
+		return SCompoundWidget::OnKeyDown(MyGeometry, InKeyboardEvent);
+	}
+
+private:
+	/** @returns Whether specified item should be expanded */
+	bool IsTreeItemExpanded(TSharedPtr<FLevelModel> Item) const
+	{
+		return Item->GetLevelExpansionFlag();
+	}
+
+	/** Appends the Level's name to the OutSearchStrings array if the Level is valid */
+	void TransformLevelToString(const TSharedPtr<FLevelModel>& Level, TArray<FString>& OutSearchStrings) const
+	{
+		if (Level.IsValid() && Level->HasValidPackage())
+		{
+			OutSearchStrings.Add(FPackageName::GetShortName(Level->GetLongPackageName()));
+		}
+	}
+	
+	/** @return Text entered in search box */
+	FText GetSearchBoxText() const
+	{
+		return SearchBoxLevelFilter->GetRawFilterText();
+	}
+
+	/** @return	Returns the filter status text */
+	FString GetFilterStatusText() const
+	{
+		const int32 SelectedLevelsCount = WorldModel->GetSelectedLevels().Num();
+		const int32 TotalLevelsCount = WorldModel->GetAllLevels().Num();
+		const int32 FilteredLevelsCount = WorldModel->GetFilteredLevels().Num();
+		
+		if (!WorldModel->IsFilterActive())
+		{
+			if (SelectedLevelsCount == 0)
+			{
+				return FString::Printf( *LOCTEXT("ShowingAllLevels", "%d levels").ToString(), TotalLevelsCount );
+			}
+			else
+			{
+				return FString::Printf( *LOCTEXT("ShowingAllLevelsSelected", "%d levels (%d selected)").ToString(), TotalLevelsCount, SelectedLevelsCount );
+			}
+		}
+		else if(WorldModel->IsFilterActive() && FilteredLevelsCount == 0)
+		{
+			return FString::Printf( *LOCTEXT("ShowingNoLevels", "No matching levels (%d total)").ToString(), TotalLevelsCount );
+		}
+		else if (SelectedLevelsCount != 0)
+		{
+			return FString::Printf( *LOCTEXT("ShowingOnlySomeLevelsSelected", "Showing %d of %d levels (%d selected)").ToString(), FilteredLevelsCount, TotalLevelsCount, SelectedLevelsCount );
+		}
+		else
+		{
+			return FString::Printf( *LOCTEXT("ShowingOnlySomeLevels", "Showing %d of %d levels").ToString(), FilteredLevelsCount, TotalLevelsCount );
+		}
+	}
+
+
+	/** @return Returns color for the filter status text message, based on success of search filter */
+	FSlateColor GetFilterStatusTextColor() const
+	{
+		if (!WorldModel->IsFilterActive() )
+		{
+			// White = no text filter
+			return FLinearColor( 1.0f, 1.0f, 1.0f );
+		}
+		else if(WorldModel->GetFilteredLevels().Num() == 0)
+		{
+			// Red = no matching actors
+			return FLinearColor( 1.0f, 0.4f, 0.4f );
+		}
+		else
+		{
+			// Green = found at least one match!
+			return FLinearColor( 0.4f, 1.0f, 0.4f );
+		}
+	}
+
+	/** @return the content for the view button */
+	TSharedRef<SWidget> GetViewButtonContent()
+	{
+		FMenuBuilder MenuBuilder(true, NULL);
+
+		MenuBuilder.BeginSection("SubLevelsViewMenu", LOCTEXT("ShowHeading", "Show"));
+		{
+			MenuBuilder.AddMenuEntry(LOCTEXT("ToggleDisplayPaths", "Display Paths"), 
+									LOCTEXT("ToggleDisplayPaths_Tooltip", "If enabled, displays the path for each level"), 
+									FSlateIcon(), 
+									FUIAction(
+										FExecuteAction::CreateSP(this, &SWorldHierarchyImpl::ToggleDisplayPaths_Executed), 
+										FCanExecuteAction(),
+										FIsActionChecked::CreateSP(this, &SWorldHierarchyImpl::GetDisplayPathsState)),
+									NAME_None, 
+									EUserInterfaceActionType::ToggleButton
+			);
+
+			MenuBuilder.AddMenuEntry(LOCTEXT("ToggleDisplayActorsCount", "Display Actors Count"), 
+									LOCTEXT("ToggleDisplayActorsCount_Tooltip", "If enabled, displays actors count for each level"), 
+									FSlateIcon(), 
+									FUIAction(
+										FExecuteAction::CreateSP(this, &SWorldHierarchyImpl::ToggleDisplayActorsCount_Executed), 
+										FCanExecuteAction(),
+										FIsActionChecked::CreateSP(this, &SWorldHierarchyImpl::GetDisplayActorsCountState)),
+									NAME_None, 
+									EUserInterfaceActionType::ToggleButton
+			);
+
+		}
+		MenuBuilder.EndSection();
+
+		return MenuBuilder.MakeWidget();
+	}
+
+	/** @return the foreground color for the view button */
+	FSlateColor GetViewButtonForegroundColor() const
+	{
+		return ViewOptionsComboButton->IsHovered() ? FEditorStyle::GetSlateColor("InvertedForeground") : FEditorStyle::GetSlateColor("DefaultForeground");
+	}
+
+	/** Toggles state of 'display path' */
+	void ToggleDisplayPaths_Executed()
+	{
+		WorldModel->SetDisplayPathsState(!WorldModel->GetDisplayPathsState());
+	}
+
+	/** Gets the state of the 'display paths' enabled/disabled */
+	bool GetDisplayPathsState() const
+	{
+		return WorldModel->GetDisplayPathsState();
+	}
+
+	/** Toggles state of 'display actors count' */
+	void ToggleDisplayActorsCount_Executed()
+	{
+		WorldModel->SetDisplayActorsCountState(!WorldModel->GetDisplayActorsCountState());
+	}
+
+	/** Gets the state of the 'display actors count' enabled/disabled */
+	bool GetDisplayActorsCountState() const
+	{
+		return WorldModel->GetDisplayActorsCountState();
+	}
+
+private:
+	/**	Whether the view is currently updating the viewmodel selection */
+	bool								bUpdatingSelection;
+	/** Our list view widget */
+	TSharedPtr<SLevelsTreeWidget>		TreeWidget;
+	
+	/** Items collection to display */
+	TSharedPtr<FLevelCollectionModel>	WorldModel;
+
+	/** The Header Row for the hierarchy */
+	TSharedPtr<SHeaderRow>				HeaderRowWidget;
+
+	/**	The LevelTextFilter that constrains which Levels appear in the hierarchy */
+	TSharedPtr<LevelTextFilter>			SearchBoxLevelFilter;
+
+	/**	Button representing view options on bottom */
+	TSharedPtr<SComboButton>			ViewOptionsComboButton;
+};
+
+//----------------------------------------------------------------
+//
+//
+//----------------------------------------------------------------
+SWorldHierarchy::SWorldHierarchy()
+{
+}
+
+SWorldHierarchy::~SWorldHierarchy()
+{
+	FWorldBrowserModule& WorldBrowserModule = FModuleManager::GetModuleChecked<FWorldBrowserModule>("WorldBrowser");
+	WorldBrowserModule.OnBrowseWorld.RemoveAll(this);
+}
+
+void SWorldHierarchy::Construct(const FArguments& InArgs)
+{
+	FWorldBrowserModule& WorldBrowserModule = FModuleManager::GetModuleChecked<FWorldBrowserModule>("WorldBrowser");
+	WorldBrowserModule.OnBrowseWorld.AddSP(this, &SWorldHierarchy::OnBrowseWorld);
+
+	OnBrowseWorld(InArgs._InWorld);
+}
+
+void SWorldHierarchy::OnBrowseWorld(UWorld* InWorld)
+{
+	// Remove all binding to an old world
+	ChildSlot
+	[
+		SNullWidget::NullWidget
+	];
+
+	WorldModel = nullptr;
+		
+	// Bind to a new world
+	if (InWorld)
+	{
+		FWorldBrowserModule& WorldBrowserModule = FModuleManager::GetModuleChecked<FWorldBrowserModule>("WorldBrowser");
+		WorldModel = WorldBrowserModule.SharedWorldModel(InWorld);
+
+		ChildSlot
+		[
+			SNew(SVerticalBox)
+
+			// Toolbar
+			+SVerticalBox::Slot()
+			.AutoHeight()
+			[
+				SNew(SBorder)
+				.BorderImage(FEditorStyle::GetBrush(TEXT("ToolPanel.GroupBorder")))
+				[
+					SNew(SHorizontalBox)
+					
+					// Toolbar
+					+SHorizontalBox::Slot()
+					.FillWidth(1.f)
+					.VAlign(VAlign_Center)
+					.HAlign(HAlign_Left)
+					[
+						// Levels menu
+						SNew( SComboButton )
+						.ComboButtonStyle(FEditorStyle::Get(), "ToolbarComboButton")
+						.ForegroundColor(FLinearColor::White)
+						.ContentPadding(0)
+						.OnGetMenuContent(this, &SWorldHierarchy::GetFileButtonContent)
+						//.ToolTipText(this, &SWorldHierarchy::GetNewAssetToolTipText)
+						//.IsEnabled(this, &SWorldHierarchy::IsAssetPathSelected )
+						.ButtonContent()
+						[
+							SNew(SHorizontalBox)
+
+							// Icon
+							+ SHorizontalBox::Slot()
+							.AutoWidth()
+							.VAlign(VAlign_Center)
+							[
+								SNew(SImage)
+								.Image(FEditorStyle::GetBrush("ContentBrowser.NewAsset"))
+							]
+
+							// Text
+							+ SHorizontalBox::Slot()
+							.AutoWidth()
+							.VAlign(VAlign_Center)
+							.Padding(0,0,2,0)
+							[
+								SNew(STextBlock)
+								.TextStyle(FEditorStyle::Get(), "ContentBrowser.TopBar.Font")
+								.Text(LOCTEXT("LevelsButton", "Levels"))
+							]
+						]
+					]
+
+					// Button to summon levels details tab
+					+SHorizontalBox::Slot()
+					.AutoWidth()
+					.VAlign(VAlign_Center)
+					.HAlign(HAlign_Right)
+					[
+						SNew(SButton)
+						.ButtonStyle(FEditorStyle::Get(), "ToggleButton")
+						.OnClicked(this, &SWorldHierarchy::OnSummonDetails)
+						.ToolTipText(LOCTEXT("SummonDetailsToolTipText", "Summons sublevels details"))
+						.HAlign(HAlign_Center)
+						.VAlign(VAlign_Center)
+						.Content()
+						[
+							SNew(SImage)
+							.Image(this, &SWorldHierarchy::GetSummonDetailsBrush)
+						]
+					]
+
+					// Button to summon world composition tab
+					+SHorizontalBox::Slot()
+					.AutoWidth()
+					.VAlign(VAlign_Center)
+					.HAlign(HAlign_Right)
+					[
+						SNew(SButton)
+						.Visibility(this, &SWorldHierarchy::GetCompositionButtonVisibility)
+						.ButtonStyle(FEditorStyle::Get(), "ToggleButton")
+						.OnClicked(this, &SWorldHierarchy::OnSummonComposition)
+						.ToolTipText(LOCTEXT("SummonHierarchyToolTipText", "Summons world composition"))
+						.HAlign(HAlign_Center)
+						.VAlign(VAlign_Center)
+						.Content()
+						[
+							SNew(SImage)
+							.Image(this, &SWorldHierarchy::GetSummonCompositionBrush)
+						]
+					]
+				]
+			]
+			
+			// Hierarchy
+			+SVerticalBox::Slot()
+			.FillHeight(1.f)
+			.Padding(0,4,0,0)
+			[
+				SNew(SBorder)
+				.BorderImage(FEditorStyle::GetBrush(TEXT("ToolPanel.GroupBorder")))
+				[
+					SNew(SWorldHierarchyImpl).InWorldModel(WorldModel)
+				]
+			]
+		];
+	}
+}
+
+FReply SWorldHierarchy::OnSummonDetails()
+{
+	FLevelEditorModule& LevelEditorModule = FModuleManager::LoadModuleChecked<FLevelEditorModule>( "LevelEditor" );
+	LevelEditorModule.SummonWorldBrowserDetails();
+	return FReply::Handled();
+}
+
+const FSlateBrush* SWorldHierarchy::GetSummonDetailsBrush() const
+{
+	return FEditorStyle::GetBrush("WorldBrowser.DetailsButtonBrush");
+}
+
+EVisibility SWorldHierarchy::GetCompositionButtonVisibility() const
+{
+	return WorldModel->IsTileWorld() ? EVisibility::Visible : EVisibility::Collapsed;
+}
+
+FReply SWorldHierarchy::OnSummonComposition()
+{
+	FLevelEditorModule& LevelEditorModule = FModuleManager::LoadModuleChecked<FLevelEditorModule>( "LevelEditor" );
+	LevelEditorModule.SummonWorldBrowserComposition();
+	return FReply::Handled();
+}
+
+const FSlateBrush* SWorldHierarchy::GetSummonCompositionBrush() const
+{
+	return FEditorStyle::GetBrush("WorldBrowser.CompositionButtonBrush");
+}
+
+TSharedRef<SWidget> SWorldHierarchy::GetFileButtonContent()
+{
+	FMenuBuilder MenuBuilder(true, WorldModel->GetCommandList());
+
+	// let current level collection model fill additional 'File' commands
+	WorldModel->CustomizeFileMainMenu(MenuBuilder);
+
+	return MenuBuilder.MakeWidget();
+}
+
+#undef LOCTEXT_NAMESPACE
\ No newline at end of file
diff --git a/Engine/Source/Editor/WorldBrowser/Private/SWorldHierarchy.h b/Engine/Source/Editor/WorldBrowser/Private/SWorldHierarchy.h
new file mode 100644
--- /dev/null
+++ b/Engine/Source/Editor/WorldBrowser/Private/SWorldHierarchy.h
@@ -0,0 +1,38 @@
+// Copyright 1998-2014 Epic Games, Inc. All Rights Reserved.
+#pragma once
+
+//----------------------------------------------------------------
+//
+//
+//----------------------------------------------------------------
+class SWorldHierarchy
+	: public SCompoundWidget
+{
+	SLATE_BEGIN_ARGS(SWorldHierarchy)
+		:_InWorld(nullptr)
+		{}
+		SLATE_ARGUMENT(UWorld*, InWorld)
+	SLATE_END_ARGS()
+
+	void Construct(const FArguments& InArgs);
+	SWorldHierarchy();
+	~SWorldHierarchy();
+
+private:
+	void OnBrowseWorld(UWorld* InWorld);
+
+	/**  */
+	FReply OnSummonDetails();
+	const FSlateBrush* GetSummonDetailsBrush() const;
+
+	/**  */
+	EVisibility GetCompositionButtonVisibility() const;
+	FReply OnSummonComposition();
+	const FSlateBrush* GetSummonCompositionBrush() const;
+
+	/**  */
+	TSharedRef<SWidget> GetFileButtonContent();
+
+private:
+	TSharedPtr<FLevelCollectionModel> WorldModel;
+};
diff --git a/Engine/Source/Editor/WorldBrowser/Private/SWorldHierarchyItem.cpp b/Engine/Source/Editor/WorldBrowser/Private/SWorldHierarchyItem.cpp
new file mode 100644
--- /dev/null
+++ b/Engine/Source/Editor/WorldBrowser/Private/SWorldHierarchyItem.cpp
@@ -0,0 +1,492 @@
+// Copyright 1998-2014 Epic Games, Inc. All Rights Reserved.
+#include "WorldBrowserPrivatePCH.h"
+
+#include "SWorldHierarchyItem.h"
+
+#define LOCTEXT_NAMESPACE "WorldBrowser"
+
+//////////////////////////
+// SWorldHierarchyItem
+//////////////////////////
+
+void SWorldHierarchyItem::Construct(const FArguments& InArgs, TSharedRef<STableViewBase> OwnerTableView)
+{
+	WorldModel = InArgs._InWorldModel;
+	LevelModel = InArgs._InItemModel;
+	IsItemExpanded = InArgs._IsItemExpanded;
+	HighlightText = InArgs._HighlightText;
+
+	StreamingLevelAlwaysLoadedBrush = FEditorStyle::GetBrush("WorldBrowser.LevelStreamingAlwaysLoaded");
+	StreamingLevelKismetBrush = FEditorStyle::GetBrush("WorldBrowser.LevelStreamingKismet");
+
+	SMultiColumnTableRow<TSharedPtr<FLevelModel>>::Construct(
+		FSuperRowType::FArguments()
+			.OnDragDetected(this, &SWorldHierarchyItem::OnItemDragDetected), 
+		OwnerTableView
+	);
+
+}
+
+BEGIN_SLATE_FUNCTION_BUILD_OPTIMIZATION
+TSharedRef< SWidget > SWorldHierarchyItem::GenerateWidgetForColumn( const FName & ColumnID )
+{
+	TSharedPtr< SWidget > TableRowContent = SNullWidget::NullWidget;
+
+	if (ColumnID == HierarchyColumns::ColumnID_LevelLabel)
+	{
+		TableRowContent =
+			SNew(SHorizontalBox)
+			
+			+SHorizontalBox::Slot()
+			.AutoWidth()
+			[
+				SNew(SExpanderArrow, SharedThis(this))
+			]
+			
+			+SHorizontalBox::Slot()
+			.VAlign(VAlign_Center)
+			.AutoWidth()
+			[
+				SNew(SBox)
+				.WidthOverride(7) // in case brush for this item is empty
+				[
+					SNew(SImage)
+					.Image(this, &SWorldHierarchyItem::GetLevelIconBrush)
+				]
+			]
+
+			+SHorizontalBox::Slot()
+			.VAlign(VAlign_Center)
+			.AutoWidth()
+			[
+				SNew(STextBlock)
+				.Font(this, &SWorldHierarchyItem::GetLevelDisplayNameFont)
+				.Text(this, &SWorldHierarchyItem::GetLevelDisplayNameText)
+				.ColorAndOpacity(this, &SWorldHierarchyItem::GetLevelDisplayNameColorAndOpacity)
+				.HighlightText(HighlightText)
+				.ToolTipText(LOCTEXT("DoubleClickToolTip", "Double-Click to make this the current Level"))
+			]
+		;
+	}
+	else if (ColumnID == HierarchyColumns::ColumnID_Lock)
+	{
+		TableRowContent = 
+			SAssignNew(LockButton, SButton)
+			.ContentPadding(0 )
+			.ButtonStyle(FEditorStyle::Get(), "ToggleButton")
+			.IsEnabled(this, &SWorldHierarchyItem::IsLockEnabled)
+			.OnClicked(this, &SWorldHierarchyItem::OnToggleLock)
+			.ToolTipText(this, &SWorldHierarchyItem::GetLevelLockToolTip)
+			.HAlign(HAlign_Center)
+			.VAlign(VAlign_Center)
+			.Content()
+			[
+				SNew(SImage)
+				.Image(this, &SWorldHierarchyItem::GetLevelLockBrush)
+			]
+		;
+	}
+	else if (ColumnID == HierarchyColumns::ColumnID_Visibility)
+	{
+		TableRowContent = 
+			SAssignNew(VisibilityButton, SButton)
+			.ContentPadding(0)
+			.ButtonStyle(FEditorStyle::Get(), "ToggleButton")
+			.IsEnabled(this, &SWorldHierarchyItem::IsVisibilityEnabled)
+			.OnClicked(this, &SWorldHierarchyItem::OnToggleVisibility)
+			.ToolTipText(LOCTEXT("VisibilityButtonToolTip", "Toggle Level Visibility"))
+			.HAlign(HAlign_Center)
+			.VAlign(VAlign_Center)
+			.Content()
+			[
+				SNew( SImage )
+				.Image(this, &SWorldHierarchyItem::GetLevelVisibilityBrush)
+			]
+		;
+	}
+	else if (ColumnID == HierarchyColumns::ColumnID_Kismet)
+	{
+		TableRowContent = 
+			SAssignNew(KismetButton, SButton)
+			.ContentPadding(0)
+			.ButtonStyle(FEditorStyle::Get(), "ToggleButton")
+			.IsEnabled(this, &SWorldHierarchyItem::IsKismetEnabled)
+			.OnClicked(this, &SWorldHierarchyItem::OnOpenKismet)
+			.ToolTipText(LOCTEXT("KismetButtonToolTip", "Open Level Blueprint"))
+			.HAlign( HAlign_Center )
+			.VAlign( VAlign_Center )
+			.Content()
+			[
+				SNew(SImage)
+				.Image(this, &SWorldHierarchyItem::GetLevelKismetBrush)
+			]
+		;
+	}
+	else if (ColumnID == HierarchyColumns::ColumnID_SCCStatus)
+	{
+		TableRowContent = 
+			SNew(SVerticalBox)
+			+SVerticalBox::Slot()
+			.FillHeight(1.0f)
+			.HAlign(HAlign_Center)
+			.VAlign(VAlign_Center)
+			[
+				SNew(SHorizontalBox)
+				+SHorizontalBox::Slot()
+				.FillWidth(1.0f)
+				.VAlign(VAlign_Center)
+				[
+					SNew( SImage )
+						.Image(this, &SWorldHierarchyItem::GetSCCStateImage)
+						.ToolTipText(this, &SWorldHierarchyItem::GetSCCStateTooltip)
+				]
+			]
+		;
+	}
+	else if (ColumnID == HierarchyColumns::ColumnID_Save)
+	{
+		TableRowContent = 
+			SAssignNew(SaveButton, SButton)
+			.ContentPadding(0)
+			.ButtonStyle(FEditorStyle::Get(), "ToggleButton")
+			.IsEnabled(this, &SWorldHierarchyItem::IsSaveEnabled)
+			.OnClicked(this, &SWorldHierarchyItem::OnSave)
+			.ToolTipText(LOCTEXT("SaveButtonToolTip", "Save Level"))
+			.HAlign( HAlign_Center )
+			.VAlign( VAlign_Center )
+			.Content()
+			[
+				SNew(SImage)
+				.Image(this, &SWorldHierarchyItem::GetLevelSaveBrush)
+			]
+		;
+	}
+
+	return TableRowContent.ToSharedRef();
+}
+END_SLATE_FUNCTION_BUILD_OPTIMIZATION
+
+FString SWorldHierarchyItem::GetLevelDisplayNameText() const
+{
+	return LevelModel->GetDisplayName();
+}
+
+bool SWorldHierarchyItem::IsSaveEnabled() const
+{
+	return LevelModel->IsLoaded();
+}
+
+bool SWorldHierarchyItem::IsLockEnabled() const
+{
+	return LevelModel->IsLoaded();
+}
+
+bool SWorldHierarchyItem::IsVisibilityEnabled() const
+{
+	return LevelModel->IsLoaded();
+}
+
+bool SWorldHierarchyItem::IsKismetEnabled() const
+{
+	return LevelModel->HasKismet();
+}
+
+FReply SWorldHierarchyItem::OnToggleVisibility()
+{
+	FLevelModelList LevelList; 
+	LevelList.Add(LevelModel);
+
+	if (LevelModel->IsVisible())
+	{
+		WorldModel->HideLevels(LevelList);
+	}
+	else
+	{
+		WorldModel->ShowLevels(LevelList);
+	}
+
+	return FReply::Handled();
+}
+
+FReply SWorldHierarchyItem::OnToggleLock()
+{
+	FLevelModelList LevelList; 
+	LevelList.Add(LevelModel);
+	
+	if (LevelModel->IsLocked())
+	{
+		WorldModel->UnlockLevels(LevelList);
+	}
+	else
+	{
+		WorldModel->LockLevels(LevelList);
+	}
+	
+	return FReply::Handled();
+}
+
+FReply SWorldHierarchyItem::OnSave()
+{
+	FLevelModelList LevelList; 
+	LevelList.Add(LevelModel);
+
+
+	WorldModel->SaveLevels(LevelList);
+	return FReply::Handled();
+}
+
+FReply SWorldHierarchyItem::OnOpenKismet()
+{
+	LevelModel->OpenKismet();
+	return FReply::Handled();
+}
+
+FReply SWorldHierarchyItem::OnItemDragDetected(const FGeometry& MyGeometry, const FPointerEvent& MouseEvent)
+{
+	if (MouseEvent.IsMouseButtonDown(EKeys::LeftMouseButton))
+	{
+		TSharedPtr<FLevelDragDropOp> Op = WorldModel->CreateDragDropOp();
+		if (Op.IsValid())
+		{
+			// Start a drag-drop
+			return FReply::Handled().BeginDragDrop(Op.ToSharedRef());
+		}
+	}
+
+	return FReply::Unhandled();
+}
+
+FReply SWorldHierarchyItem::OnDrop( const FGeometry& MyGeometry, const FDragDropEvent& DragDropEvent )
+{
+	auto Op = DragDropEvent.GetOperationAs<FLevelDragDropOp>();
+	if (Op.IsValid() && LevelModel->IsGoodToDrop(Op))
+	{
+		LevelModel->OnDrop(Op);
+		return FReply::Handled();
+	}
+
+	return FReply::Unhandled();
+}
+
+void SWorldHierarchyItem::OnDragEnter( const FGeometry& MyGeometry, const FDragDropEvent& DragDropEvent )
+{
+	SCompoundWidget::OnDragEnter(MyGeometry, DragDropEvent);
+	
+	auto Op = DragDropEvent.GetOperationAs<FLevelDragDropOp>();
+	if (Op.IsValid())
+	{
+		// to show mouse hover effect
+		SWidget::OnMouseEnter(MyGeometry, DragDropEvent);
+		// D&D decorator icon
+		Op->bGoodToDrop = LevelModel->IsGoodToDrop(Op);
+	}
+}
+
+void SWorldHierarchyItem::OnDragLeave( const FDragDropEvent& DragDropEvent )
+{
+	SCompoundWidget::OnDragLeave(DragDropEvent);
+
+	auto Op = DragDropEvent.GetOperationAs<FLevelDragDropOp>();
+	if (Op.IsValid())
+	{
+		// to hide mouse hover effect
+		SWidget::OnMouseLeave(DragDropEvent);
+		// D&D decorator icon
+		Op->bGoodToDrop = false;
+	}
+}
+
+FSlateFontInfo SWorldHierarchyItem::GetLevelDisplayNameFont() const
+{
+	if (LevelModel->IsCurrent())
+	{
+		return FEditorStyle::GetFontStyle("LevelBrowser.LabelFontBold");
+	}
+	else
+	{
+		return FEditorStyle::GetFontStyle("LevelBrowser.LabelFont");
+	}
+}
+
+FSlateColor SWorldHierarchyItem::GetLevelDisplayNameColorAndOpacity() const
+{
+	// Force the text to display red if level is missing
+	if (!LevelModel->HasValidPackage())
+	{
+		return FLinearColor(1.0f, 0.0f, 0.0f);
+	}
+		
+	// Highlight text differently if it doesn't match the search filter (e.g., parent levels to child levels that
+	// match search criteria.)
+	if (LevelModel->GetLevelFilteredOutFlag())
+	{
+		return FLinearColor(0.30f, 0.30f, 0.30f);
+	}
+
+	if (!LevelModel->IsLoaded())
+	{
+		return FSlateColor::UseSubduedForeground();
+	}
+		
+	if (LevelModel->IsCurrent())
+	{
+		return LevelModel->GetLevelSelectionFlag() ? FSlateColor::UseForeground() : FLinearColor(0.12f, 0.56f, 1.0f);
+	}
+
+	return FSlateColor::UseForeground();
+}
+
+const FSlateBrush* SWorldHierarchyItem::GetLevelIconBrush() const
+{
+	UClass* StreamingClass = LevelModel->GetStreamingClass();
+	if (StreamingClass == ULevelStreamingKismet::StaticClass())
+	{
+		return StreamingLevelKismetBrush;
+	}
+
+	if (StreamingClass == ULevelStreamingAlwaysLoaded::StaticClass())
+	{
+		return StreamingLevelAlwaysLoadedBrush;
+	}
+
+	return nullptr;
+}
+
+const FSlateBrush* SWorldHierarchyItem::GetLevelVisibilityBrush() const
+{
+	if (LevelModel->GetLevelObject())
+	{
+		if (LevelModel->IsVisible())
+		{
+			return VisibilityButton->IsHovered() ? FEditorStyle::GetBrush( "Level.VisibleHighlightIcon16x" ) :
+													FEditorStyle::GetBrush( "Level.VisibleIcon16x" );
+		}
+		else
+		{
+			return VisibilityButton->IsHovered() ? FEditorStyle::GetBrush( "Level.NotVisibleHighlightIcon16x" ) :
+													FEditorStyle::GetBrush( "Level.NotVisibleIcon16x" );
+		}
+	}
+	else
+	{
+		return FEditorStyle::GetBrush( "Level.EmptyIcon16x" );
+	}
+}
+
+const FSlateBrush* SWorldHierarchyItem::GetLevelLockBrush() const
+{
+	if (!LevelModel->IsLoaded() || LevelModel->IsPersistent())
+	{
+		//Locking the persistent level is not allowed; stub in a different brush
+		return FEditorStyle::GetBrush( "Level.EmptyIcon16x" );
+	}
+	else 
+	{
+		////Non-Persistent
+		//if ( GEngine && GEngine->bLockReadOnlyLevels )
+		//{
+		//	if(LevelModel->IsReadOnly())
+		//	{
+		//		return LockButton->IsHovered() ? FEditorStyle::GetBrush( "Level.ReadOnlyLockedHighlightIcon16x" ) :
+		//											FEditorStyle::GetBrush( "Level.ReadOnlyLockedIcon16x" );
+		//	}
+		//}
+			
+		if (LevelModel->IsLocked())
+		{
+			return LockButton->IsHovered() ? FEditorStyle::GetBrush( "Level.LockedHighlightIcon16x" ) :
+												FEditorStyle::GetBrush( "Level.LockedIcon16x" );
+		}
+		else
+		{
+			return LockButton->IsHovered() ? FEditorStyle::GetBrush( "Level.UnlockedHighlightIcon16x" ) :
+												FEditorStyle::GetBrush( "Level.UnlockedIcon16x" );
+		}
+	}
+}
+
+FString SWorldHierarchyItem::GetLevelLockToolTip() const
+{
+	//Non-Persistent
+	if (GEngine && GEngine->bLockReadOnlyLevels)
+	{
+		if (LevelModel->IsFileReadOnly())
+		{
+			return LOCTEXT("ReadOnly_LockButtonToolTip", "Read-Only levels are locked!").ToString();
+		}
+	}
+
+	return LOCTEXT("LockButtonToolTip", "Toggle Level Lock").ToString();
+}
+
+FString SWorldHierarchyItem::GetSCCStateTooltip() const
+{
+	FSourceControlStatePtr SourceControlState = ISourceControlModule::Get().GetProvider().GetState(LevelModel->GetPackageFileName(), EStateCacheUsage::Use);
+	if(SourceControlState.IsValid())
+	{
+		return SourceControlState->GetDisplayTooltip().ToString();
+	}
+
+	return FString();
+}
+
+const FSlateBrush* SWorldHierarchyItem::GetSCCStateImage() const
+{
+	FSourceControlStatePtr SourceControlState = ISourceControlModule::Get().GetProvider().GetState(LevelModel->GetPackageFileName(), EStateCacheUsage::Use);
+	if(SourceControlState.IsValid())
+	{
+		return FEditorStyle::GetBrush(SourceControlState->GetSmallIconName());
+	}
+
+	return NULL;
+}
+
+const FSlateBrush* SWorldHierarchyItem::GetLevelSaveBrush() const
+{
+	if (LevelModel->IsLoaded())
+	{
+		if (LevelModel->IsLocked())
+		{
+			return FEditorStyle::GetBrush( "Level.SaveDisabledIcon16x" );
+		}
+		else
+		{
+			if (LevelModel->IsDirty())
+			{
+				return SaveButton->IsHovered() ? FEditorStyle::GetBrush( "Level.SaveModifiedHighlightIcon16x" ) :
+													FEditorStyle::GetBrush( "Level.SaveModifiedIcon16x" );
+			}
+			else
+			{
+				return SaveButton->IsHovered() ? FEditorStyle::GetBrush( "Level.SaveHighlightIcon16x" ) :
+													FEditorStyle::GetBrush( "Level.SaveIcon16x" );
+			}
+		}								
+	}
+	else
+	{
+		return FEditorStyle::GetBrush( "Level.EmptyIcon16x" );
+	}	
+}
+
+const FSlateBrush* SWorldHierarchyItem::GetLevelKismetBrush() const
+{
+	if (LevelModel->IsLoaded())
+	{
+		if (LevelModel->HasKismet())
+		{
+			return KismetButton->IsHovered() ? FEditorStyle::GetBrush( "Level.ScriptHighlightIcon16x" ) :
+												FEditorStyle::GetBrush( "Level.ScriptIcon16x" );
+		}
+		else
+		{
+			return FEditorStyle::GetBrush( "Level.EmptyIcon16x" );
+		}
+	}
+	else
+	{
+		return FEditorStyle::GetBrush( "Level.EmptyIcon16x" );
+	}	
+}
+
+#undef LOCTEXT_NAMESPACE
diff --git a/Engine/Source/Editor/WorldBrowser/Private/SWorldHierarchyItem.h b/Engine/Source/Editor/WorldBrowser/Private/SWorldHierarchyItem.h
new file mode 100644
--- /dev/null
+++ b/Engine/Source/Editor/WorldBrowser/Private/SWorldHierarchyItem.h
@@ -0,0 +1,175 @@
+// Copyright 1998-2014 Epic Games, Inc. All Rights Reserved.
+#pragma once
+
+namespace HierarchyColumns
+{
+	/** IDs for list columns */
+	static const FName ColumnID_LevelLabel( "Level" );
+	static const FName ColumnID_Visibility( "Visibility" );
+	static const FName ColumnID_Lock( "Lock" );
+	static const FName ColumnID_SCCStatus( "SCC_Status" );
+	static const FName ColumnID_Save( "Save" );
+	static const FName ColumnID_Kismet( "Blueprint" );
+	static const FName ColumnID_ActorCount( "ActorCount" );
+	static const FName ColumnID_LightmassSize( "LightmassSize" );
+	static const FName ColumnID_FileSize( "FileSize" );
+}
+
+/** A single item in the levels hierarchy tree. Represents a level model */
+class SWorldHierarchyItem 
+	: public SMultiColumnTableRow<TSharedPtr<FLevelModel>>
+{
+public:
+	DECLARE_DELEGATE_TwoParams( FOnNameChanged, const TSharedPtr<FLevelModel>& /*TreeItem*/, const FVector2D& /*MessageLocation*/);
+
+	SLATE_BEGIN_ARGS( SWorldHierarchyItem )
+		: _IsItemExpanded( false )
+	{}
+		/** Data for the world */
+		SLATE_ARGUMENT(TSharedPtr<FLevelCollectionModel>, InWorldModel)
+	
+		/** Item model this widget represents */
+		SLATE_ARGUMENT(TSharedPtr<FLevelModel>, InItemModel)
+
+		/** True when this item has children and is expanded */
+		SLATE_ATTRIBUTE(bool, IsItemExpanded)
+
+		/** The string in the title to highlight */
+		SLATE_ATTRIBUTE(FText, HighlightText)
+
+	SLATE_END_ARGS()
+
+	/** Constructs this widget with InArgs */
+	void Construct(const FArguments& InArgs, TSharedRef<STableViewBase> OwnerTableView);
+
+	TSharedRef<SWidget> GenerateWidgetForColumn( const FName& ColumnName ) OVERRIDE;
+
+	FReply OnItemDragDetected(const FGeometry& MyGeometry, const FPointerEvent& MouseEvent);
+	FReply OnDrop( const FGeometry& MyGeometry, const FDragDropEvent& DragDropEvent ) OVERRIDE;
+	void OnDragEnter( const FGeometry& MyGeometry, const FDragDropEvent& DragDropEvent ) OVERRIDE;
+	void OnDragLeave( const FDragDropEvent& DragDropEvent ) OVERRIDE;
+
+private:
+	/** Operations buttons enabled/disabled state */
+	bool IsSaveEnabled() const;
+	bool IsLockEnabled() const;
+	bool IsVisibilityEnabled() const;
+	bool IsKismetEnabled() const;
+		
+	/**
+	 *	Called when the user clicks on the visibility icon for a Level's item widget
+	 *
+	 *	@return	A reply that indicated whether this event was handled.
+	 */
+	FReply OnToggleVisibility();
+
+	/**
+	 *	Called when the user clicks on the lock icon for a Level's item widget
+	 *
+	 *	@return	A reply that indicated whether this event was handled.
+	 */
+	FReply OnToggleLock();
+
+	/**
+	 *	Called when the user clicks on the save icon for a Level's item widget
+	 *
+	 *	@return	A reply that indicated whether this event was handled.
+	 */
+	FReply OnSave();
+
+	/**
+	 *	Called when the user clicks on the kismet icon for a Level's item widget
+	 *
+	 *	@return	A reply that indicated whether this event was handled.
+	 */
+	FReply OnOpenKismet();
+	
+	/**
+	 *  @return The text of level name while it is not being edited
+	 */
+	FString GetLevelDisplayNameText() const;
+	
+	/** */
+	FSlateFontInfo GetLevelDisplayNameFont() const;
+	
+	/** */
+	FSlateColor GetLevelDisplayNameColorAndOpacity() const;
+
+	/**
+	 *	@return	The SlateBrush representing level icon
+	 */
+	const FSlateBrush* GetLevelIconBrush() const;
+
+	/**
+	 *	Called to get the Slate Image Brush representing the visibility state of
+	 *	the Level this item widget represents
+	 *
+	 *	@return	The SlateBrush representing the Level's visibility state
+	 */
+	const FSlateBrush* GetLevelVisibilityBrush() const;
+
+	/**
+	 *	Called to get the Slate Image Brush representing the lock state of
+	 *	the Level this item widget represents
+	 *
+	 *	@return	The SlateBrush representing the Level's lock state
+	 */
+	const FSlateBrush* GetLevelLockBrush() const;
+
+	/**
+	 *	Called to get the Slate Image Brush representing the save state of
+	 *	the Level this row widget represents
+	 *
+	 *	@return	The SlateBrush representing the Level's save state
+	 */
+	const FSlateBrush* GetLevelSaveBrush() const;
+
+	/**
+	 *	Called to get the Slate Image Brush representing the kismet state of
+	 *	the Level this row widget represents
+	 *
+	 *	@return	The SlateBrush representing the Level's kismet state
+	 */
+	const FSlateBrush* GetLevelKismetBrush() const;
+	
+	/** */
+	FString GetLevelLockToolTip() const;
+	
+	/** */
+	FString GetSCCStateTooltip() const;
+	
+	/** */
+	const FSlateBrush* GetSCCStateImage() const;
+
+private:
+	/** The world data */
+	TSharedPtr<FLevelCollectionModel>	WorldModel;
+
+	/** The data for this item */
+	TSharedPtr<FLevelModel>			LevelModel;
+	
+	/** The string to highlight in level display name */
+	TAttribute<FText>				HighlightText;
+
+	/** True when this item has children and is expanded */
+	TAttribute<bool>				IsItemExpanded;
+
+	/** Brushes for the different streaming level types */
+	const FSlateBrush*				StreamingLevelAlwaysLoadedBrush;
+	const FSlateBrush*				StreamingLevelKismetBrush;
+
+	/**	The visibility button for the Level */
+	TSharedPtr<SButton>				VisibilityButton;
+
+	/**	The lock button for the Level */
+	TSharedPtr<SButton>				LockButton;
+
+	/**	The save button for the Level */
+	TSharedPtr<SButton>				SaveButton;
+
+	/**	The kismet button for the Level */
+	TSharedPtr<SButton>				KismetButton;
+
+	/**	The color button for the Level */
+	TSharedPtr<SButton>				ColorButton;
+};
diff --git a/Engine/Source/Editor/WorldBrowser/Private/SWorldMainView.cpp b/Engine/Source/Editor/WorldBrowser/Private/SWorldMainView.cpp
deleted file mode 100644
diff --git a/Engine/Source/Editor/WorldBrowser/Private/SWorldMainView.h b/Engine/Source/Editor/WorldBrowser/Private/SWorldMainView.h
deleted file mode 100644
diff --git a/Engine/Source/Editor/WorldBrowser/Private/SWorldTreeItem.cpp b/Engine/Source/Editor/WorldBrowser/Private/SWorldTreeItem.cpp
deleted file mode 100644
diff --git a/Engine/Source/Editor/WorldBrowser/Private/SWorldTreeItem.h b/Engine/Source/Editor/WorldBrowser/Private/SWorldTreeItem.h
deleted file mode 100644
diff --git a/Engine/Source/Editor/WorldBrowser/Private/SWorldTreeView.cpp b/Engine/Source/Editor/WorldBrowser/Private/SWorldTreeView.cpp
deleted file mode 100644
diff --git a/Engine/Source/Editor/WorldBrowser/Private/SWorldTreeView.h b/Engine/Source/Editor/WorldBrowser/Private/SWorldTreeView.h
deleted file mode 100644
diff --git a/Engine/Source/Editor/WorldBrowser/Private/StreamingLevels/StreamingLevelCollectionModel.h b/Engine/Source/Editor/WorldBrowser/Private/StreamingLevels/StreamingLevelCollectionModel.h
--- a/Engine/Source/Editor/WorldBrowser/Private/StreamingLevels/StreamingLevelCollectionModel.h
+++ b/Engine/Source/Editor/WorldBrowser/Private/StreamingLevels/StreamingLevelCollectionModel.h
@@ -26,8 +26,6 @@ public:
 		return LevelCollectionModel;
 	}
 
-	virtual bool SupportsGridView() const OVERRIDE { return false; };
-
 public:
 	/** FLevelCollection interface */
 	virtual void UnloadLevels(const FLevelModelList& InLevelList) OVERRIDE;
@@ -36,7 +34,6 @@ public:
 	virtual void CustomizeFileMainMenu(FMenuBuilder& InMenuBuilder) const OVERRIDE;
 	virtual void RegisterDetailsCustomization(class FPropertyEditorModule& PropertyModule, TSharedPtr<class IDetailsView> InDetailsView)  OVERRIDE;
 	virtual void UnregisterDetailsCustomization(class FPropertyEditorModule& PropertyModule, TSharedPtr<class IDetailsView> InDetailsView) OVERRIDE;
-	
 
 private:
 	virtual void Initialize() OVERRIDE;
diff --git a/Engine/Source/Editor/WorldBrowser/Private/StreamingLevels/StreamingLevelCustomization.cpp b/Engine/Source/Editor/WorldBrowser/Private/StreamingLevels/StreamingLevelCustomization.cpp
--- a/Engine/Source/Editor/WorldBrowser/Private/StreamingLevels/StreamingLevelCustomization.cpp
+++ b/Engine/Source/Editor/WorldBrowser/Private/StreamingLevels/StreamingLevelCustomization.cpp
@@ -7,7 +7,7 @@
 #include "SVectorInputBox.h"
 #include "SRotatorInputBox.h"
 
-#include "SWorldDetailsView.h"
+#include "SWorldDetails.h"
 #include "StreamingLevelCollectionModel.h"
 #include "StreamingLevelCustomization.h"
 #include "StreamingLevelEdMode.h"
diff --git a/Engine/Source/Editor/WorldBrowser/Private/Tiles/SWorldComposition.cpp b/Engine/Source/Editor/WorldBrowser/Private/Tiles/SWorldComposition.cpp
new file mode 100644
--- /dev/null
+++ b/Engine/Source/Editor/WorldBrowser/Private/Tiles/SWorldComposition.cpp
@@ -0,0 +1,1222 @@
+// Copyright 1998-2014 Epic Games, Inc. All Rights Reserved.
+#include "WorldBrowserPrivatePCH.h"
+#include "SNodePanel.h"
+#include "GraphEditor.h"
+#include "EdGraphUtilities.h"
+
+#include "SWorldComposition.h"
+#include "SWorldTileItem.h"
+#include "SWorldLayers.h"
+#include "WorldTileCollectionModel.h"
+
+#define LOCTEXT_NAMESPACE "WorldBrowser"
+
+//----------------------------------------------------------------
+//
+//
+//----------------------------------------------------------------
+struct FWorldZoomLevelsContainer 
+	: public FZoomLevelsContainer
+{
+	float	GetZoomAmount(int32 InZoomLevel) const OVERRIDE
+	{
+		return 1.f/FMath::Square(GetNumZoomLevels() - InZoomLevel + 1)*2.f;
+	}
+
+	int32 GetNearestZoomLevel( float InZoomAmount ) const OVERRIDE
+	{
+		for (int32 ZoomLevelIndex=0; ZoomLevelIndex < GetNumZoomLevels(); ++ZoomLevelIndex)
+		{
+			if (InZoomAmount <= GetZoomAmount(ZoomLevelIndex))
+			{
+				return ZoomLevelIndex;
+			}
+		}
+
+		return GetDefaultZoomLevel();
+	}
+
+	FText GetZoomText(int32 InZoomLevel) const OVERRIDE
+	{
+		return FText::AsNumber(GetZoomAmount(InZoomLevel));
+	}
+
+	int32	GetNumZoomLevels() const OVERRIDE
+	{
+		return 100;
+	}
+
+	int32	GetDefaultZoomLevel() const OVERRIDE
+	{
+		return GetNumZoomLevels() - 10;
+	}
+
+	EGraphRenderingLOD::Type GetLOD(int32 InZoomLevel) const OVERRIDE
+	{
+		return EGraphRenderingLOD::DefaultDetail;
+	}
+};
+
+//----------------------------------------------------------------
+//
+//
+//----------------------------------------------------------------
+class SWorldCompositionGrid 
+	: public SNodePanel
+{
+public:
+	SLATE_BEGIN_ARGS(SWorldCompositionGrid) {}
+		SLATE_ARGUMENT(TSharedPtr<FWorldTileCollectionModel>, InWorldModel)
+	SLATE_END_ARGS()
+
+	SWorldCompositionGrid()
+		: CommandList(MakeShareable(new FUICommandList))
+		, bHasScrollRequest(false)
+		, bIsFirstTickCall(true)
+		, bHasNodeInteraction(true)
+		, BoundsSnappingDistance(20.f)
+	{
+
+		SharedThumbnailRT = new FSlateTextureRenderTarget2DResource(
+						FLinearColor::Black, 
+						512, 
+						512, 
+						PF_B8G8R8A8, SF_Point, TA_Wrap, TA_Wrap, 0.0f
+					);
+		BeginInitResource(SharedThumbnailRT);
+	}
+
+	~SWorldCompositionGrid()
+	{
+		BeginReleaseResource(SharedThumbnailRT);
+		FlushRenderingCommands();
+		delete SharedThumbnailRT;
+		
+		WorldModel->SelectionChanged.RemoveAll(this);
+		WorldModel->CollectionChanged.RemoveAll(this);
+	}
+
+	void Construct(const FArguments& InArgs)
+	{
+		ZoomLevels = new FWorldZoomLevelsContainer();
+
+		SNodePanel::Construct();
+	
+		//// bind commands
+		//const FWorldTileCommands& Commands = FWorldTileCommands::Get();
+		//FUICommandList& ActionList = *CommandList;
+
+		//ActionList.MapAction(Commands.FitToSelection,
+		//	FExecuteAction::CreateSP(this, &SWorldCompositionGrid::FitToSelection_Executed),
+		//	FCanExecuteAction::CreateSP(this, &SWorldCompositionGrid::AreAnyItemsSelected));
+
+		//ActionList.MapAction(Commands.MoveLevelLeft,
+		//	FExecuteAction::CreateSP(this, &SWorldCompositionGrid::MoveLevelLeft_Executed),
+		//	FCanExecuteAction::CreateSP(this, &SWorldCompositionGrid::AreAnyItemsSelected));
+
+		//ActionList.MapAction(Commands.MoveLevelRight,
+		//	FExecuteAction::CreateSP(this, &SWorldCompositionGrid::MoveLevelRight_Executed),
+		//	FCanExecuteAction::CreateSP(this, &SWorldCompositionGrid::AreAnyItemsSelected));
+
+		//ActionList.MapAction(Commands.MoveLevelUp,
+		//	FExecuteAction::CreateSP(this, &SWorldCompositionGrid::MoveLevelUp_Executed),
+		//	FCanExecuteAction::CreateSP(this, &SWorldCompositionGrid::AreAnyItemsSelected));
+
+		//ActionList.MapAction(Commands.MoveLevelDown,
+		//	FExecuteAction::CreateSP(this, &SWorldCompositionGrid::MoveLevelDown_Executed),
+		//	FCanExecuteAction::CreateSP(this, &SWorldCompositionGrid::AreAnyItemsSelected));
+
+		//
+		BounceCurve.AddCurve(0.0f, 1.0f);
+		BounceCurve.Play();
+		WorldModel = InArgs._InWorldModel;
+		bUpdatingSelection = false;
+	
+		WorldModel->SelectionChanged.AddSP(this, &SWorldCompositionGrid::OnUpdateSelection);
+		WorldModel->CollectionChanged.AddSP(this, &SWorldCompositionGrid::RefreshView);
+		SelectionManager.OnSelectionChanged.BindSP(this, &SWorldCompositionGrid::OnSelectionChanged);
+	
+		RefreshView();
+	}
+	
+	/**  Add specified item to the grid view */
+	void AddItem(const TSharedPtr<FLevelModel>& LevelModel)
+	{
+		auto NewNode = SNew(SWorldTileItem)
+							.InWorldModel(WorldModel)
+							.InItemModel(LevelModel)
+							.ThumbnailRenderTarget(SharedThumbnailRT);
+	
+		AddGraphNode(NewNode);
+	}
+	
+	/**  Remove specified item from the grid view */
+	void RemoveItem(const TSharedPtr<FLevelModel>& LevelModel)
+	{
+		TSharedRef<SNode>* pItem = NodeToWidgetLookup.Find(LevelModel->GetNodeObject());
+		if (pItem == NULL)
+		{
+			return;
+		}
+
+		Children.Remove(*pItem);
+		VisibleChildren.Remove(*pItem);
+		NodeToWidgetLookup.Remove(LevelModel->GetNodeObject());
+	}
+		
+	/**  Updates all the items in the grid view */
+	void RefreshView()
+	{
+		RemoveAllNodes();
+
+		FLevelModelList AllLevels = WorldModel->GetAllLevels();
+		for (auto It = AllLevels.CreateConstIterator(); It; ++It)
+		{
+			AddItem(*It);
+		}
+	}
+		
+	/**  SWidget interface */
+	virtual void Tick(const FGeometry& AllottedGeometry, const double InCurrentTime, const float InDeltaTime) OVERRIDE
+	{
+		SNodePanel::Tick(AllottedGeometry, InCurrentTime, InDeltaTime);
+
+		// scroll to world center on first open
+		if (bIsFirstTickCall)
+		{
+			bIsFirstTickCall = false;
+			ViewOffset-= AllottedGeometry.Size*0.5f/GetZoomAmount();
+		}
+
+		FVector2D CursorPosition = FSlateApplication::Get().GetCursorPos();
+
+		// Update cached variables
+		WorldMouseLocation = CursorToWorldPosition(AllottedGeometry, CursorPosition);
+		WorldMarqueeSize = Marquee.Rect.GetSize()/AllottedGeometry.Scale;
+			
+		// Update streaming preview data
+		const bool bShowPotentiallyVisibleLevels = FSlateApplication::Get().GetModifierKeys().IsAltDown() && 
+													AllottedGeometry.IsUnderLocation(CursorPosition);
+	
+		//WorldModel->UpdateStreamingPreview(WorldMouseLocation, bShowPotentiallyVisibleLevels);
+			
+		// deffered scroll and zooming requests
+		if (bHasScrollRequest)
+		{
+			bHasScrollRequest = false;
+
+			// zoom to
+			FVector2D SizeWithZoom = RequestedZoomArea*ZoomLevels->GetZoomAmount(ZoomLevel);
+			if (RequestedAllowZoomIn || 
+				(SizeWithZoom.X >= AllottedGeometry.Size.X || SizeWithZoom.Y >= AllottedGeometry.Size.Y))
+			{
+				// maximum zoom out by default
+				ZoomLevel = ZoomLevels->GetDefaultZoomLevel();
+				// expand zoom area little bit, so zooming will fit original area not so tight
+				RequestedZoomArea*= 1.2f;
+				// find more suitable zoom value
+				for (int32 Zoom = 0; Zoom < ZoomLevels->GetDefaultZoomLevel(); ++Zoom)
+				{
+					SizeWithZoom = RequestedZoomArea*ZoomLevels->GetZoomAmount(Zoom);
+					if (SizeWithZoom.X >= AllottedGeometry.Size.X || SizeWithZoom.Y >= AllottedGeometry.Size.Y)
+					{
+						ZoomLevel = Zoom;
+						break;
+					}
+				}
+			}
+
+			// scroll to
+			ViewOffset = RequestedScrollLocation - AllottedGeometry.Size*0.5f/GetZoomAmount();
+		}
+	}
+	
+	/**  SWidget interface */
+	virtual void ArrangeChildren( const FGeometry& AllottedGeometry, FArrangedChildren& ArrangedChildren ) const OVERRIDE
+	{
+		for (int32 ChildIndex=0; ChildIndex < VisibleChildren.Num(); ++ChildIndex)
+		{
+			const auto Child = StaticCastSharedRef<SWorldTileItem>(VisibleChildren[ChildIndex]);
+			const EVisibility ChildVisibility = Child->GetVisibility();
+
+			if (ArrangedChildren.Accepts(ChildVisibility))
+			{
+				FVector2D ChildPos = Child->GetPosition();
+					
+				ArrangedChildren.AddWidget(ChildVisibility,
+					AllottedGeometry.MakeChild(Child,
+					ChildPos - GetViewOffset(),
+					Child->GetDesiredSize(), GetZoomAmount()
+					));
+			}
+		}
+	}
+	
+	/**  SWidget interface */
+	virtual int32 OnPaint(const FGeometry& AllottedGeometry, const FSlateRect& MyClippingRect, 
+							FSlateWindowElementList& OutDrawElements, int32 LayerId, 
+							const FWidgetStyle& InWidgetStyle, bool bParentEnabled ) const OVERRIDE
+	{
+		// First paint the background
+		{
+			LayerId = PaintBackground(AllottedGeometry, MyClippingRect, OutDrawElements, LayerId);
+			LayerId++;
+		}
+
+		FArrangedChildren ArrangedChildren(EVisibility::Visible);
+		ArrangeChildren(AllottedGeometry, ArrangedChildren);
+
+		// Draw the child nodes
+
+		// When drawing a marquee, need a preview of what the selection will be.
+		const auto* SelectionToVisualize = &(SelectionManager.SelectedNodes);
+		FGraphPanelSelectionSet SelectionPreview;
+		if (Marquee.IsValid())
+		{			
+			ApplyMarqueeSelection(Marquee, SelectionManager.SelectedNodes, SelectionPreview);
+			SelectionToVisualize = &SelectionPreview;
+		}
+	
+		int32 NodesLayerId = LayerId;
+
+		for (int32 ChildIndex = 0; ChildIndex < ArrangedChildren.Num(); ++ChildIndex)
+		{
+			FArrangedWidget& CurWidget = ArrangedChildren(ChildIndex);
+			TSharedRef<SWorldTileItem> ChildNode = StaticCastSharedRef<SWorldTileItem>(CurWidget.Widget);
+		
+			ChildNode->bAffectedByMarquee = SelectionToVisualize->Contains(ChildNode->GetObjectBeingDisplayed());
+			LayerId = CurWidget.Widget->OnPaint(CurWidget.Geometry, MyClippingRect, OutDrawElements, NodesLayerId, InWidgetStyle, ShouldBeEnabled(bParentEnabled));
+			ChildNode->bAffectedByMarquee = false;
+		}
+		
+		// Draw editable world bounds
+		if (!WorldModel->IsSimulating())
+		{
+			float ScreenSpaceSize = FLevelCollectionModel::EditableAxisLength()*GetZoomAmount()*2.f;
+			FVector2D PaintSize = FVector2D(ScreenSpaceSize, ScreenSpaceSize);
+			FVector2D PaintPosition = GraphCoordToPanelCoord(FVector2D::ZeroVector) - (PaintSize*0.5f);
+			FPaintGeometry EditableArea = AllottedGeometry.ToPaintGeometry(PaintPosition, PaintSize);
+			EditableArea.DrawScale = 0.2f;
+			FLinearColor PaintColor = FLinearColor::Red;
+
+			FSlateDrawElement::MakeBox(
+				OutDrawElements,
+				++LayerId,
+				EditableArea,
+				FEditorStyle::GetBrush(TEXT("Graph.CompactNode.ShadowSelected")),
+				MyClippingRect,
+				ESlateDrawEffect::None,
+				PaintColor
+				);
+		}
+		
+		// Draw the marquee selection rectangle
+		PaintMarquee(AllottedGeometry, MyClippingRect, OutDrawElements, ++LayerId);
+
+		// Draw the software cursor
+		PaintSoftwareCursor(AllottedGeometry, MyClippingRect, OutDrawElements, ++LayerId);
+
+		if(WorldModel->IsSimulating())
+		{
+			// Draw a surrounding indicator when PIE is active, to make it clear that the graph is read-only, etc...
+			FSlateDrawElement::MakeBox(
+				OutDrawElements,
+				LayerId,
+				AllottedGeometry.ToPaintGeometry(),
+				FEditorStyle::GetBrush(TEXT("Graph.PlayInEditor")),
+				MyClippingRect
+				);
+
+			// Draw a current camera position
+			{
+				const FSlateBrush* CameraImage = FEditorStyle::GetBrush(TEXT("WorldBrowser.SimulationViewPositon"));
+				FVector ObserverPosition = WorldModel->GetObserverPosition();
+				FVector2D ObserverPositionScreen = GraphCoordToPanelCoord(FVector2D(ObserverPosition.X, ObserverPosition.Y));
+
+				FPaintGeometry PaintGeometry = AllottedGeometry.ToPaintGeometry(
+					ObserverPositionScreen - CameraImage->ImageSize*0.5f, 
+					CameraImage->ImageSize
+					);
+						
+				FSlateDrawElement::MakeBox(
+					OutDrawElements,
+					++LayerId,
+					PaintGeometry,
+					CameraImage,
+					MyClippingRect
+				);
+			}
+		}
+
+		LayerId = PaintScaleRuler(AllottedGeometry, MyClippingRect, OutDrawElements, LayerId);
+		return LayerId;
+	}
+		
+	/** SWidget interface */
+	virtual FReply OnMouseMove(const FGeometry& MyGeometry, const FPointerEvent& MouseEvent) OVERRIDE
+	{
+		const bool bIsRightMouseButtonDown = MouseEvent.IsMouseButtonDown( EKeys::RightMouseButton );
+		const bool bIsLeftMouseButtonDown = MouseEvent.IsMouseButtonDown( EKeys::LeftMouseButton );
+
+		PastePosition = PanelCoordToGraphCoord( MyGeometry.AbsoluteToLocal( MouseEvent.GetScreenSpacePosition() ) );
+
+		if ( this->HasMouseCapture() )
+		{
+			const FVector2D CursorDelta = MouseEvent.GetCursorDelta();
+			// Track how much the mouse moved since the mouse down.
+			TotalMouseDelta += CursorDelta.Size();
+
+			if (bIsRightMouseButtonDown)
+			{
+				FReply ReplyState = FReply::Handled();
+
+				if( !CursorDelta.IsZero() )
+				{
+					bShowSoftwareCursor = true;
+				}
+
+				// Panning and mouse is outside of panel? Pasting should just go to the screen center.
+				PastePosition = PanelCoordToGraphCoord( 0.5 * MyGeometry.Size );
+
+				this->bIsPanning = true;
+				ViewOffset -= CursorDelta / GetZoomAmount();
+
+				return ReplyState;
+			}
+			else if (bIsLeftMouseButtonDown)
+			{
+				TSharedPtr<SNode> NodeBeingDragged = NodeUnderMousePtr.Pin();
+
+				if ( IsEditable.Get() )
+				{
+					// Update the amount to pan panel
+					UpdateViewOffset(MyGeometry, MouseEvent.GetScreenSpacePosition());
+
+					const bool bCursorInDeadZone = TotalMouseDelta <= SlatePanTriggerDistance;
+
+					if ( NodeBeingDragged.IsValid() )
+					{
+						if ( !bCursorInDeadZone )
+						{
+							// Note, NodeGrabOffset() comes from the node itself, so it's already scaled correctly.
+							FVector2D AnchorNodeNewPos = PanelCoordToGraphCoord( MyGeometry.AbsoluteToLocal( MouseEvent.GetScreenSpacePosition() ) ) - NodeGrabOffset;
+
+							// Dragging an unselected node automatically selects it.
+							SelectionManager.StartDraggingNode(NodeBeingDragged->GetObjectBeingDisplayed(), MouseEvent);
+
+							// Move all the selected nodes.
+							{
+								const FVector2D AnchorNodeOldPos = NodeBeingDragged->GetPosition();
+								const FVector2D DeltaPos = AnchorNodeNewPos - AnchorNodeOldPos;
+								if (DeltaPos.Size() > KINDA_SMALL_NUMBER)
+								{
+									MoveSelectedNodes(NodeBeingDragged, AnchorNodeNewPos);
+								}
+							}
+						}
+
+						return FReply::Handled();
+					}
+				}
+
+				if ( !NodeBeingDragged.IsValid() )
+				{
+					// We are marquee selecting
+					const FVector2D GraphMousePos = PanelCoordToGraphCoord( MyGeometry.AbsoluteToLocal( MouseEvent.GetScreenSpacePosition() ) );
+					Marquee.Rect.UpdateEndPoint(GraphMousePos);
+
+					FindNodesAffectedByMarquee( /*out*/ Marquee.AffectedNodes );
+					return FReply::Handled();
+				}
+			}
+		}
+
+		return FReply::Unhandled();
+	}
+
+	/** @return Size of a marquee rectangle in world space */
+	FVector2D GetMarqueeWorldSize() const
+	{
+		return WorldMarqueeSize;
+	}
+
+	/** @return Mouse cursor position in world space */
+	FVector2D GetMouseWorldLocation() const
+	{
+		return WorldMouseLocation;
+	}
+
+protected:
+	/**  Draws background for grid view */
+	uint32 PaintBackground(const FGeometry& AllottedGeometry, const FSlateRect& MyClippingRect, FSlateWindowElementList& OutDrawElements, uint32 LayerId) const
+	{
+		FVector2D ScreenWorldOrigin = GraphCoordToPanelCoord(FVector2D(0, 0));
+		FSlateRect ScreenRect = AllottedGeometry.GetRect();
+	
+		// Vertical line
+		if (ScreenWorldOrigin.X > ScreenRect.Left &&
+			ScreenWorldOrigin.X < ScreenRect.Right)
+		{
+			TArray<FVector2D> LinePoints;
+			LinePoints.Add(FVector2D(ScreenWorldOrigin.X, ScreenRect.Top));
+			LinePoints.Add(FVector2D(ScreenWorldOrigin.X, ScreenRect.Bottom));
+		
+			FSlateDrawElement::MakeLines( 
+				OutDrawElements,
+				LayerId,
+				AllottedGeometry.ToPaintGeometry(),
+				LinePoints,
+				MyClippingRect,
+				ESlateDrawEffect::None,
+				FColor(100, 100, 100));
+		}
+
+		// Horizontal line
+		if (ScreenWorldOrigin.Y > ScreenRect.Top &&
+			ScreenWorldOrigin.Y < ScreenRect.Bottom)
+		{
+			TArray<FVector2D> LinePoints;
+			LinePoints.Add(FVector2D(ScreenRect.Left, ScreenWorldOrigin.Y));
+			LinePoints.Add(FVector2D(ScreenRect.Right, ScreenWorldOrigin.Y));
+		
+			FSlateDrawElement::MakeLines( 
+				OutDrawElements,
+				LayerId,
+				AllottedGeometry.ToPaintGeometry(),
+				LinePoints,
+				MyClippingRect,
+				ESlateDrawEffect::None,
+				FColor(100, 100, 100));
+		}
+
+		return LayerId + 1;
+	}
+
+	/**  Draws current scale */
+	uint32 PaintScaleRuler(const FGeometry& AllottedGeometry, const FSlateRect& MyClippingRect, FSlateWindowElementList& OutDrawElements, uint32 LayerId) const
+	{
+		const float	ScaleRulerLength = 100.f; // pixels
+		TArray<FVector2D> LinePoints;
+		LinePoints.Add(FVector2D::ZeroVector);
+		LinePoints.Add(FVector2D::ZeroVector + FVector2D(ScaleRulerLength, 0.f));
+	
+		FSlateDrawElement::MakeLines( 
+			OutDrawElements,
+			LayerId,
+			AllottedGeometry.ToOffsetPaintGeometry(FVector2D(10, 40)),
+			LinePoints,
+			MyClippingRect,
+			ESlateDrawEffect::None,
+			FColor(200, 200, 200));
+
+		const float UnitsInRuler = ScaleRulerLength/GetZoomAmount();// Pixels to world units
+		const int32 UnitsInMeter = 100;
+		const int32 UnitsInKilometer = UnitsInMeter*1000;
+	
+		FString RulerText;
+		if (UnitsInRuler > UnitsInKilometer) // in kilometers
+		{
+			RulerText = FString::Printf(TEXT("%.2f km"), UnitsInRuler/UnitsInKilometer);
+		}
+		else // in meters
+		{
+			RulerText = FString::Printf(TEXT("%.2f m"), UnitsInRuler/UnitsInMeter);
+		}
+	
+		FSlateDrawElement::MakeText(
+			OutDrawElements,
+			LayerId,
+			AllottedGeometry.ToOffsetPaintGeometry(FVector2D(10, 27)),
+			RulerText,
+			FEditorStyle::GetFontStyle("NormalFont"),
+			MyClippingRect,
+			ESlateDrawEffect::None,
+			FColor(200, 200, 200));
+	
+		return LayerId + 1;
+	}
+		
+	/**  SWidget interface */
+	virtual FReply OnKeyDown(const FGeometry& MyGeometry, const FKeyboardEvent& InKeyboardEvent) OVERRIDE
+	{
+		if (CommandList->ProcessCommandBindings(InKeyboardEvent))
+		{
+			return FReply::Handled();
+		}
+	
+		if (WorldModel->GetCommandList()->ProcessCommandBindings(InKeyboardEvent))
+		{
+			return FReply::Handled();
+		}
+
+		return SNodePanel::OnKeyDown(MyGeometry, InKeyboardEvent);
+	}
+
+	/**  SWidget interface */
+	bool SupportsKeyboardFocus() const OVERRIDE
+	{
+		return true;
+	}
+	
+	/**  SNodePanel interface */
+	TSharedPtr<SWidget> OnSummonContextMenu(const FGeometry& MyGeometry, const FPointerEvent& MouseEvent)
+	{
+		if (WorldModel->IsReadOnly())
+		{
+			return SNullWidget::NullWidget;
+		}
+
+		FArrangedChildren ArrangedChildren(EVisibility::Visible);
+		ArrangeChildren(MyGeometry, ArrangedChildren);
+
+		const int32 NodeUnderMouseIndex = SWidget::FindChildUnderMouse( ArrangedChildren, MouseEvent );
+		if (NodeUnderMouseIndex != INDEX_NONE)
+		{
+			// PRESSING ON A NODE!
+			const FArrangedWidget& NodeGeometry = ArrangedChildren(NodeUnderMouseIndex);
+			const FVector2D MousePositionInNode = NodeGeometry.Geometry.AbsoluteToLocal(MouseEvent.GetScreenSpacePosition());
+			TSharedRef<SNode> NodeWidgetUnderMouse = StaticCastSharedRef<SNode>( NodeGeometry.Widget );
+
+			if (NodeWidgetUnderMouse->CanBeSelected(MousePositionInNode))
+			{
+				if (!SelectionManager.IsNodeSelected(NodeWidgetUnderMouse->GetObjectBeingDisplayed()))
+				{
+					SelectionManager.SelectSingleNode(NodeWidgetUnderMouse->GetObjectBeingDisplayed());
+				}
+			}
+		}
+		else
+		{
+			SelectionManager.ClearSelectionSet();
+		}
+	
+		// Summon context menu
+		FMenuBuilder MenuBuilder(true, WorldModel->GetCommandList());
+		WorldModel->BuildGridMenu(MenuBuilder);
+		TSharedPtr<SWidget> MenuWidget = MenuBuilder.MakeWidget();
+
+		FSlateApplication::Get().PushMenu(
+			AsShared(),
+			MenuWidget.ToSharedRef(),
+			MouseEvent.GetScreenSpacePosition(),
+			FPopupTransitionEffect( FPopupTransitionEffect::ContextMenu )
+			);
+
+		return MenuWidget;
+	}
+	
+	/**  SNodePanel interface */
+	virtual void PopulateVisibleChildren(const FGeometry& AllottedGeometry) OVERRIDE
+	{
+		VisibleChildren.Empty();
+
+		FSlateRect PanelRect = AllottedGeometry.GetRect();
+		FVector2D ViewStartPos = PanelCoordToGraphCoord(FVector2D(PanelRect.Left, PanelRect.Top));
+		FVector2D ViewEndPos = PanelCoordToGraphCoord(FVector2D(PanelRect.Right, PanelRect.Bottom));
+		FSlateRect ViewRect(ViewStartPos, ViewEndPos);
+
+		for (int32 ChildIndex=0; ChildIndex < Children.Num(); ++ChildIndex)
+		{
+			const auto Child = StaticCastSharedRef<SWorldTileItem>(Children[ChildIndex]);
+		
+			if (WorldModel->PassesAllFilters(Child->GetLevelModel()))
+			{
+				FSlateRect ChildRect = Child->GetItemRect();
+				FVector2D ChildSize = ChildRect.GetSize();
+
+				if (ChildSize.X > 0.f && 
+					ChildSize.Y > 0.f && 
+					FSlateRect::DoRectanglesIntersect(ChildRect, ViewRect))
+				{
+					VisibleChildren.Add(Child);
+				}
+			}
+		}
+
+		// Sort tiles such that smaller and selected tiles will be drawn on top of other tiles
+		struct FVisibleTilesSorter
+		{
+			FVisibleTilesSorter(const FLevelCollectionModel& InWorldModel) : WorldModel(InWorldModel)
+			{}
+			bool operator()(const TSharedRef<SNodePanel::SNode>& A,
+							const TSharedRef<SNodePanel::SNode>& B) const
+			{
+				TSharedRef<SWorldTileItem> ItemA = StaticCastSharedRef<SWorldTileItem>(A);
+				TSharedRef<SWorldTileItem> ItemB = StaticCastSharedRef<SWorldTileItem>(B);
+				return WorldModel.CompareLevelsZOrder(ItemA->GetLevelModel(), ItemB->GetLevelModel());
+			}
+			const FLevelCollectionModel& WorldModel;
+		};
+
+		VisibleChildren.Sort(FVisibleTilesSorter(*WorldModel.Get()));
+	}
+
+	/**  SNodePanel interface */
+	virtual void OnBeginNodeInteraction(const TSharedRef<SNode>& InNodeToDrag, const FVector2D& GrabOffset) OVERRIDE
+	{
+		bHasNodeInteraction = true;
+		SNodePanel::OnBeginNodeInteraction(InNodeToDrag, GrabOffset);
+	}
+	
+	/**  SNodePanel interface */
+	virtual void OnEndNodeInteraction(const TSharedRef<SNode>& InNodeDragged) OVERRIDE
+	{
+		const SWorldTileItem& Item = static_cast<const SWorldTileItem&>(InNodeDragged.Get());
+		if (Item.IsItemEditable())
+		{
+			FVector2D AbsoluteDelta = Item.GetLevelModel()->GetLevelTranslationDelta();
+			FIntPoint IntAbsoluteDelta = FIntPoint(AbsoluteDelta.X, AbsoluteDelta.Y);
+
+			// Reset stored translation delta to 0
+			WorldModel->UpdateTranslationDelta(WorldModel->GetSelectedLevels(), FVector2D::ZeroVector, false, 0.f);
+	
+			// In case we have non zero dragging delta, translate selected levels 
+			if (IntAbsoluteDelta != FIntPoint::ZeroValue)
+			{
+				WorldModel->TranslateLevels(WorldModel->GetSelectedLevels(), IntAbsoluteDelta);
+			}
+		}
+	
+		bHasNodeInteraction = false;
+
+		SNodePanel::OnEndNodeInteraction(InNodeDragged);
+	}
+
+	/** Handles selection changes in the grid view */
+	void OnSelectionChanged(const FGraphPanelSelectionSet& SelectedNodes)
+	{
+		if (bUpdatingSelection)
+		{
+			return;
+		}
+	
+		bUpdatingSelection = true;
+		FLevelModelList SelectedLevels;
+	
+		for (FGraphPanelSelectionSet::TConstIterator NodeIt(SelectedNodes); NodeIt; ++NodeIt)
+		{
+			TSharedRef<SNode>* pWidget = NodeToWidgetLookup.Find(*NodeIt);
+			if (pWidget != NULL)
+			{
+				TSharedRef<SWorldTileItem> Item = StaticCastSharedRef<SWorldTileItem>(*pWidget);
+				SelectedLevels.Add(Item->GetLevelModel());
+			}
+		}
+	
+		WorldModel->SetSelectedLevels(SelectedLevels);
+		bUpdatingSelection = false;
+	}
+	
+	/** Handles selection changes in data source */
+	void OnUpdateSelection()
+	{
+		if (bUpdatingSelection)
+		{
+			return;
+		}
+
+		bUpdatingSelection = true;
+
+		SelectionManager.ClearSelectionSet();
+		FLevelModelList SelectedLevels = WorldModel->GetSelectedLevels();
+		for (auto It = SelectedLevels.CreateConstIterator(); It; ++It)
+		{
+			SelectionManager.SetNodeSelection((*It)->GetNodeObject(), true);
+		}
+
+		if (SelectionManager.AreAnyNodesSelected())
+		{
+			FVector2D MinCorner, MaxCorner;
+			if (GetBoundsForNodes(true, MinCorner, MaxCorner, 0.f))
+			{
+				RequestScrollTo((MaxCorner + MinCorner)*0.5f, MaxCorner - MinCorner);
+			}
+		}
+		bUpdatingSelection = false;
+	}
+
+	/** Handles new item added to data source */
+	void OnNewItemAdded(TSharedPtr<FLevelModel> NewItem)
+	{
+		RefreshView();
+	}
+
+	/**  FitToSelection command handler */
+	void FitToSelection_Executed()
+	{
+		FVector2D MinCorner, MaxCorner;
+		if (GetBoundsForNodes(true, MinCorner, MaxCorner, 0.f))
+		{
+			RequestScrollTo((MaxCorner + MinCorner)*0.5f, MaxCorner - MinCorner, true);
+		}
+	}
+		
+	/**  @returns Whether any of the levels are selected */
+	bool AreAnyItemsSelected() const
+	{
+		return SelectionManager.AreAnyNodesSelected();
+	}
+
+	/**  Requests view scroll to specified position and fit to specified area 
+	 *   @param	InLocation		The location to scroll to
+	 *   @param	InArea			The area to fit in view
+	 *   @param	bAllowZoomIn	Is zoom in allowed during fit to area calculations
+	 */
+	void RequestScrollTo(FVector2D InLocation, FVector2D InArea, bool bAllowZoomIn = false)
+	{
+		bHasScrollRequest = true;
+		RequestedScrollLocation = InLocation;
+		RequestedZoomArea = InArea;
+		RequestedAllowZoomIn = bAllowZoomIn;
+	}
+	
+	/** Handlers for moving items using arrow keys */
+	void MoveLevelLeft_Executed()
+	{
+		if (!bHasNodeInteraction)
+		{
+			WorldModel->TranslateLevels(WorldModel->GetSelectedLevels(), FIntPoint(-1, 0));
+		}
+	}
+
+	void MoveLevelRight_Executed()
+	{
+		if (!bHasNodeInteraction)
+		{
+			WorldModel->TranslateLevels(WorldModel->GetSelectedLevels(), FIntPoint(+1, 0));
+		}
+	}
+
+	void MoveLevelUp_Executed()
+	{
+		if (!bHasNodeInteraction)
+		{
+			WorldModel->TranslateLevels(WorldModel->GetSelectedLevels(), FIntPoint(0, -1));
+		}
+	}
+
+	void MoveLevelDown_Executed()
+	{
+		if (!bHasNodeInteraction)
+		{
+			WorldModel->TranslateLevels(WorldModel->GetSelectedLevels(), FIntPoint(0, +1));
+		}
+	}	
+	
+	/** Moves selected nodes by specified offset */
+	void MoveSelectedNodes(const TSharedPtr<SNode>& InNodeToDrag, FVector2D NewPosition)
+	{
+		auto ItemDragged = StaticCastSharedPtr<SWorldTileItem>(InNodeToDrag);
+	
+		if (ItemDragged->IsItemEditable())
+		{
+			// Current translation snapping value
+			float SnappingDistanceWorld = 0.f;
+			const bool bBoundsSnapping = (FSlateApplication::Get().GetModifierKeys().IsControlDown() == false);
+			if (bBoundsSnapping)
+			{
+				SnappingDistanceWorld = BoundsSnappingDistance/GetZoomAmount();
+			}
+			else if (GetDefault<ULevelEditorViewportSettings>()->GridEnabled)
+			{
+				SnappingDistanceWorld = GEditor->GetGridSize();
+			}
+		
+			FVector2D StartPosition = ItemDragged->GetPosition() - ItemDragged->GetLevelModel()->GetLevelTranslationDelta();
+			FVector2D AbsoluteDelta = NewPosition - StartPosition;
+
+			WorldModel->UpdateTranslationDelta(WorldModel->GetSelectedLevels(), AbsoluteDelta, bBoundsSnapping, SnappingDistanceWorld);
+		}
+	}
+
+	/**  Converts cursor absolute position to the world position */
+	FVector2D CursorToWorldPosition(const FGeometry& InGeometry, FVector2D InAbsoluteCursorPosition)
+	{
+		FVector2D ViewSpacePosition = (InAbsoluteCursorPosition - InGeometry.AbsolutePosition)/InGeometry.Scale;
+		return PanelCoordToGraphCoord(ViewSpacePosition);
+	}
+
+private:
+	/** Levels data list to display*/
+	TSharedPtr<FWorldTileCollectionModel>	WorldModel;
+
+	/** Geometry cache */
+	mutable FVector2D						CachedAllottedGeometryScaledSize;
+	/** Bouncing curve */
+	FCurveSequence							BounceCurve;
+	bool									bUpdatingSelection;
+	TArray<FIntRect>						OccupiedCells;
+	const TSharedRef<FUICommandList>		CommandList;
+
+	bool									bHasScrollRequest;
+	FVector2D								RequestedScrollLocation;
+	FVector2D								RequestedZoomArea;
+	bool									RequestedAllowZoomIn;
+
+	bool									bIsFirstTickCall;
+	// Is user interacting with a node now
+	bool									bHasNodeInteraction;
+
+	// Snapping distance in screen units for a tile bounds
+	float									BoundsSnappingDistance;
+
+	//
+	// Mouse location in the world
+	FVector2D								WorldMouseLocation;
+	// Current marquee rectangle size in world units
+	FVector2D								WorldMarqueeSize;
+
+	FSlateTextureRenderTarget2DResource*	SharedThumbnailRT;
+};
+
+//----------------------------------------------------------------
+//
+//
+//----------------------------------------------------------------
+SWorldComposition::SWorldComposition()
+{
+}
+
+SWorldComposition::~SWorldComposition()
+{
+	FWorldBrowserModule& WorldBrowserModule = FModuleManager::GetModuleChecked<FWorldBrowserModule>("WorldBrowser");
+	WorldBrowserModule.OnBrowseWorld.RemoveAll(this);
+	
+	OnBrowseWorld(nullptr);
+}
+
+void SWorldComposition::Construct(const FArguments& InArgs)
+{
+	FWorldBrowserModule& WorldBrowserModule = FModuleManager::GetModuleChecked<FWorldBrowserModule>("WorldBrowser");
+	WorldBrowserModule.OnBrowseWorld.AddSP(this, &SWorldComposition::OnBrowseWorld);
+
+	ChildSlot
+	[
+		SAssignNew(ContentParent, SBorder)
+		.BorderImage(FEditorStyle::GetBrush(TEXT("ToolPanel.GroupBorder")))
+	];
+	
+	OnBrowseWorld(InArgs._InWorld);
+}
+
+void SWorldComposition::OnBrowseWorld(UWorld* InWorld)
+{
+	// Remove old world bindings
+	ContentParent->SetContent(SNullWidget::NullWidget);
+	LayersListWrapBox = nullptr;
+	NewLayerButton = nullptr;
+	NewLayerPopupWindow = nullptr;
+	GridView = nullptr;
+	TileWorldModel = nullptr;
+			
+	// Bind to a new world model in case it's a world composition
+	if (InWorld && InWorld->WorldComposition)
+	{
+		// Get the shared world model for this world object
+		FWorldBrowserModule& WorldBrowserModule = FModuleManager::GetModuleChecked<FWorldBrowserModule>("WorldBrowser");
+		auto SharedWorldModel = WorldBrowserModule.SharedWorldModel(InWorld);
+		
+		// double check we have a tile world
+		if (SharedWorldModel->IsTileWorld())
+		{
+			TileWorldModel = StaticCastSharedPtr<FWorldTileCollectionModel>(SharedWorldModel);
+			ContentParent->SetContent(ConstructContentWidget());
+			PopulateLayersList();
+		}
+	}
+}
+
+TSharedRef<SWidget> SWorldComposition::ConstructContentWidget()
+{
+	return 	
+		SNew(SVerticalBox)
+
+		// Layers list
+		+SVerticalBox::Slot()
+		.AutoHeight()
+		[
+			SAssignNew(LayersListWrapBox, SWrapBox)
+			.UseAllottedWidth(true)
+		]
+				
+		+SVerticalBox::Slot()
+		.FillHeight(1)
+		[
+			SNew( SOverlay )
+
+			// Grid view
+			+SOverlay::Slot()
+			[
+				SAssignNew(GridView, SWorldCompositionGrid)
+					.InWorldModel(TileWorldModel)
+			]
+
+			// Grid view top status bar
+			+SOverlay::Slot()
+			.VAlign(VAlign_Top)
+			[
+				SNew(SBorder)
+				.BorderImage(FEditorStyle::GetBrush(TEXT("Graph.TitleBackground")))
+				[
+					SNew(SVerticalBox)
+
+					+SVerticalBox::Slot()
+					.AutoHeight()
+					[
+						SNew(SHorizontalBox)
+							
+						// Current world view scale
+						+SHorizontalBox::Slot()
+						.AutoWidth()
+						[
+							SNullWidget::NullWidget
+						]
+						+SHorizontalBox::Slot()
+						.Padding(5,0,0,0)
+						[
+							SNullWidget::NullWidget
+						]
+							
+						// World origin position
+						+SHorizontalBox::Slot()
+						.AutoWidth()
+						[
+							SNew(SImage)
+							.Image(FEditorStyle::GetBrush( "WorldBrowser.WorldOrigin" ))
+						]
+						+SHorizontalBox::Slot()
+						.Padding(5,0,0,0)
+						[
+							SNew(STextBlock)
+							.TextStyle( FEditorStyle::Get(), "WorldBrowser.StatusBarText" )
+							.Text(this, &SWorldComposition::GetCurrentOriginText)
+						]
+
+						// Current level
+						+SHorizontalBox::Slot()
+						.HAlign(HAlign_Right)
+						.Padding(0,0,5,0)
+						[
+							SNew(STextBlock)
+							.TextStyle( FEditorStyle::Get(), "WorldBrowser.StatusBarText" )
+							.Text(this, &SWorldComposition::GetCurrentLevelText)
+						]											
+					]
+				]
+			]
+			// Grid view bottom status bar
+			+SOverlay::Slot()
+			.VAlign(VAlign_Bottom)
+			[
+				SNew(SBorder)
+				.BorderImage(FEditorStyle::GetBrush(TEXT("Graph.TitleBackground")))
+				[
+					SNew(SVerticalBox)
+
+					+SVerticalBox::Slot()
+					.AutoHeight()
+					[
+						SNew(SHorizontalBox)
+
+						// Mouse location
+						+SHorizontalBox::Slot()
+						.AutoWidth()
+						[
+							SNew(SImage)
+							.Image(FEditorStyle::GetBrush( "WorldBrowser.MouseLocation" ))
+						]
+						+SHorizontalBox::Slot()
+						.Padding(5,0,0,0)
+						[
+							SNew(STextBlock)
+							.TextStyle( FEditorStyle::Get(), "WorldBrowser.StatusBarText" )
+							.Text(this, &SWorldComposition::GetMouseLocationText)
+						]
+
+						// Selection marquee rectangle size
+						+SHorizontalBox::Slot()
+						.AutoWidth()
+						[
+							SNew(SImage)
+							.Image(FEditorStyle::GetBrush( "WorldBrowser.MarqueeRectSize" ))
+						]
+						+SHorizontalBox::Slot()
+						.Padding(5,0,0,0)
+						[
+							SNew(STextBlock)
+							.TextStyle( FEditorStyle::Get(), "WorldBrowser.StatusBarText" )
+							.Text(this, &SWorldComposition::GetMarqueeSelectionSizeText)
+						]
+
+						// World size
+						+SHorizontalBox::Slot()
+						.HAlign(HAlign_Right)
+						[
+							SNew(SHorizontalBox)
+
+							+SHorizontalBox::Slot()
+							.AutoWidth()
+							[
+								SNew(SImage)
+								.Image(FEditorStyle::GetBrush( "WorldBrowser.WorldSize" ))
+							]
+
+							+SHorizontalBox::Slot()
+							.Padding(5,0,5,0)
+							[
+								SNew(STextBlock)
+								.TextStyle( FEditorStyle::Get(), "WorldBrowser.StatusBarText" )
+								.Text(this, &SWorldComposition::GetWorldSizeText)
+							]
+						]											
+					]
+				]
+			]
+
+			// Top-right corner text indicating that simulation is active
+			+SOverlay::Slot()
+			.Padding(20)
+			.VAlign(VAlign_Top)
+			.HAlign(HAlign_Right)
+			[
+				SNew(STextBlock)
+				.Visibility(this, &SWorldComposition::IsSimulationVisible)
+				.TextStyle( FEditorStyle::Get(), "Graph.SimulatingText" )
+				.Text(FString(TEXT("SIMULATING")))
+			]
+		];
+}
+
+void SWorldComposition::PopulateLayersList()
+{
+	TArray<FWorldTileLayer>& AllLayers = TileWorldModel->GetLayers();
+	
+	LayersListWrapBox->ClearChildren();
+	for (auto WorldLayer : AllLayers)
+	{
+		LayersListWrapBox->AddSlot()
+		.Padding(1,1,0,0)
+		[
+			SNew(SWorldLayerButton)
+				.InWorldModel(TileWorldModel)
+				.WorldLayer(WorldLayer)
+		];
+	}
+
+	// Add new layer button
+	LayersListWrapBox->AddSlot()
+	.Padding(1,1,0,0)
+	[
+		SAssignNew(NewLayerButton, SButton)
+		.OnClicked(this, &SWorldComposition::NewLayer_Clicked)
+		.ButtonColorAndOpacity(FLinearColor(0.2f, 0.2f, 0.2f, 0.2f))
+		.Content()
+		[
+			SNew(SImage)
+			.Image(FEditorStyle::GetBrush("WorldBrowser.AddLayer"))
+		]
+	];
+}
+
+FReply SWorldComposition::NewLayer_Clicked()
+{
+	if (TileWorldModel->IsReadOnly())
+	{
+		return FReply::Handled();
+	}
+	
+	TSharedRef<SNewLayerPopup> CreateLayerWidget = 
+		SNew(SNewLayerPopup)
+		.OnCreateLayer(this, &SWorldComposition::CreateNewLayer)
+		.DefaultName(LOCTEXT("Layer_DefaultName", "MyLayer").ToString());
+
+	NewLayerPopupWindow = FSlateApplication::Get().PushMenu(
+		this->AsShared(),
+		CreateLayerWidget,
+		FSlateApplication::Get().GetCursorPos(),
+		FPopupTransitionEffect( FPopupTransitionEffect::TypeInPopup )
+		);
+
+	return FReply::Handled();
+}
+
+FReply SWorldComposition::CreateNewLayer(const FWorldTileLayer& NewLayer)
+{
+	TileWorldModel->AddManagedLayer(NewLayer);
+	PopulateLayersList();
+	
+	if (NewLayerPopupWindow.IsValid())
+	{
+		NewLayerPopupWindow->RequestDestroyWindow();
+	}
+		
+	return FReply::Handled();
+}
+
+FText SWorldComposition::GetZoomText() const
+{
+	return GridView->GetZoomText();
+}
+
+FString SWorldComposition::GetCurrentOriginText() const
+{
+	UWorld* CurrentWorld = (TileWorldModel->IsSimulating() ? TileWorldModel->GetSimulationWorld() : TileWorldModel->GetWorld());
+	return FString::Printf(TEXT("%d, %d"), CurrentWorld->GlobalOriginOffset.X, CurrentWorld->GlobalOriginOffset.Y);
+}
+
+FString SWorldComposition::GetCurrentLevelText() const
+{
+	FString Text("");
+	UWorld* CurrentWorld = (TileWorldModel->IsSimulating() ? TileWorldModel->GetSimulationWorld() : TileWorldModel->GetWorld());
+
+	if (CurrentWorld->GetCurrentLevel())
+	{
+		UPackage* Package = CastChecked<UPackage>(CurrentWorld->GetCurrentLevel()->GetOutermost());
+		Text+= FPackageName::GetShortName(Package->GetName());
+	}
+	else
+	{
+		Text+= TEXT("None");
+	}
+	
+	return Text;
+}
+
+FString SWorldComposition::GetMouseLocationText() const
+{
+	FVector2D MouseLocation = GridView->GetMouseWorldLocation();
+	return FString::Printf(TEXT("%d, %d"), FMath::RoundToInt(MouseLocation.X), FMath::RoundToInt(MouseLocation.Y));
+}
+
+FString	SWorldComposition::GetMarqueeSelectionSizeText() const
+{
+	FVector2D MarqueeSize = GridView->GetMarqueeWorldSize();
+	if (MarqueeSize.Size() > 0)
+	{
+		return FString::Printf(TEXT("%d x %d"), FMath::RoundToInt(MarqueeSize.X), FMath::RoundToInt(MarqueeSize.Y));
+	}
+	else
+	{
+		return FString();
+	}
+}
+
+FString SWorldComposition::GetWorldSizeText() const
+{
+	FIntPoint WorldSize = TileWorldModel->GetWorldSize();
+	if (WorldSize.Size() > 0)
+	{
+		return FString::Printf(TEXT("%d x %d"), WorldSize.X, WorldSize.Y);
+	}
+	else
+	{
+		return FString();
+	}
+}
+
+EVisibility SWorldComposition::IsSimulationVisible() const
+{
+	return (TileWorldModel->IsSimulating() ? EVisibility::Visible : EVisibility::Hidden);
+}
+
+#undef LOCTEXT_NAMESPACE
\ No newline at end of file
diff --git a/Engine/Source/Editor/WorldBrowser/Private/Tiles/SWorldComposition.h b/Engine/Source/Editor/WorldBrowser/Private/Tiles/SWorldComposition.h
new file mode 100644
--- /dev/null
+++ b/Engine/Source/Editor/WorldBrowser/Private/Tiles/SWorldComposition.h
@@ -0,0 +1,56 @@
+// Copyright 1998-2014 Epic Games, Inc. All Rights Reserved.
+#pragma once
+
+class SWorldComposition 
+	: public SCompoundWidget
+{
+public:
+	SLATE_BEGIN_ARGS(SWorldComposition)
+		:_InWorld(nullptr)
+		{}
+		SLATE_ARGUMENT(UWorld*, InWorld)
+	SLATE_END_ARGS()
+
+	SWorldComposition();
+	~SWorldComposition();
+
+	void Construct(const FArguments& InArgs);
+
+private:
+	/**  */
+	void OnBrowseWorld(UWorld* InWorld);
+
+	/**  */
+	TSharedRef<SWidget> ConstructContentWidget();
+
+	/** Populate current FWorldTileLayer list to UI */
+	void PopulateLayersList();
+
+	/** Creates a popup window with New layer parameters */
+	FReply NewLayer_Clicked();
+	
+	/** Creates a new managed layer */
+	FReply CreateNewLayer(const FWorldTileLayer& NewLayer);
+
+	/** Top status bar details */
+	FText GetZoomText() const;
+	FString	GetCurrentOriginText() const;
+	FString GetCurrentLevelText() const;
+
+	/** Bottom status bar details */
+	FString GetMouseLocationText() const;
+	FString	GetMarqueeSelectionSizeText() const;
+	FString GetWorldSizeText() const;
+
+	/** @return whether SIMULATION sign should be visible */
+	EVisibility IsSimulationVisible() const;
+	
+private:
+	TSharedPtr<class FWorldTileCollectionModel> TileWorldModel;
+	
+	TSharedPtr<SBorder>							ContentParent;
+	TSharedPtr<SWrapBox>						LayersListWrapBox;
+	TSharedPtr<SButton>							NewLayerButton;
+	TSharedPtr<class SWindow>					NewLayerPopupWindow;
+	TSharedPtr<class SWorldCompositionGrid>		GridView;
+};
diff --git a/Engine/Source/Editor/WorldBrowser/Private/Tiles/SWorldTileItem.cpp b/Engine/Source/Editor/WorldBrowser/Private/Tiles/SWorldTileItem.cpp
new file mode 100644
--- /dev/null
+++ b/Engine/Source/Editor/WorldBrowser/Private/Tiles/SWorldTileItem.cpp
@@ -0,0 +1,454 @@
+// Copyright 1998-2014 Epic Games, Inc. All Rights Reserved.
+#include "WorldBrowserPrivatePCH.h"
+
+#include "ObjectTools.h"
+#include "EdGraphUtilities.h"
+#include "SWorldTileItem.h"
+
+#define LOCTEXT_NAMESPACE "WorldBrowser"
+
+FTileItemThumbnail::FTileItemThumbnail(FSlateTextureRenderTarget2DResource* InThumbnailRenderTarget, 
+									   TSharedPtr<FLevelModel> InItemModel)
+	: LevelModel(InItemModel)
+	, ThumbnailTexture(NULL)
+	, ThumbnailRenderTarget(InThumbnailRenderTarget)
+{
+	FIntPoint RTSize = ThumbnailRenderTarget->GetSizeXY();
+	ThumbnailTexture = new FSlateTexture2DRHIRef(RTSize.X, RTSize.Y, PF_B8G8R8A8, NULL, TexCreate_Dynamic);
+	BeginInitResource(ThumbnailTexture);
+}
+
+FTileItemThumbnail::~FTileItemThumbnail()
+{
+	BeginReleaseResource(ThumbnailTexture);
+	
+	// Wait for all resources to be released
+	FlushRenderingCommands();
+	delete ThumbnailTexture;
+	ThumbnailTexture = NULL;
+}
+
+FIntPoint FTileItemThumbnail::GetSize() const
+{
+	return ThumbnailRenderTarget->GetSizeXY();
+}
+
+FSlateShaderResource* FTileItemThumbnail::GetViewportRenderTargetTexture() const
+{
+	return ThumbnailTexture;
+}
+
+bool FTileItemThumbnail::RequiresVsync() const
+{
+	return false;
+}
+
+void FTileItemThumbnail::UpdateTextureData(FObjectThumbnail* ObjectThumbnail)
+{
+	check(ThumbnailTexture)
+	
+	if (ObjectThumbnail &&
+		ObjectThumbnail->GetImageWidth() > 0 && 
+		ObjectThumbnail->GetImageHeight() > 0 && 
+		ObjectThumbnail->GetUncompressedImageData().Num() > 0)
+	{
+		// Make bulk data for updating the texture memory later
+		FSlateTextureDataPtr ThumbnailBulkData = MakeShareable(new FSlateTextureData(
+			ObjectThumbnail->GetImageWidth(), 
+			ObjectThumbnail->GetImageHeight(), 
+			GPixelFormats[PF_B8G8R8A8].BlockBytes, 
+			ObjectThumbnail->AccessImageData())
+			);
+					
+		// Update the texture RHI
+		ENQUEUE_UNIQUE_RENDER_COMMAND_TWOPARAMETER(
+			UpdateGridItemThumbnailResourse,
+			FSlateTexture2DRHIRef*, Texture, ThumbnailTexture,
+			FSlateTextureDataPtr, BulkData, ThumbnailBulkData,
+		{
+			Texture->SetTextureData(BulkData);
+			Texture->UpdateRHI();
+		});
+	}
+}
+
+void FTileItemThumbnail::UpdateThumbnail()
+{
+	// No need images for persistent and always loaded levels
+	if (LevelModel->IsPersistent())
+	{
+		return;
+	}
+	
+	// Load image from a package header
+	if (!LevelModel->IsVisible())
+	{
+		TSet<FName> ObjectFullNames;
+		ObjectFullNames.Add(LevelModel->GetAssetName());
+		FThumbnailMap ThumbnailMap;
+			
+		if (ThumbnailTools::ConditionallyLoadThumbnailsFromPackage(LevelModel->GetPackageFileName(), ObjectFullNames, ThumbnailMap))
+		{
+			UpdateTextureData(ThumbnailMap.Find(LevelModel->GetAssetName()));
+		}
+	}
+	// Render image from a loaded visble level
+	else
+	{
+		ULevel* TargetLevel = LevelModel->GetLevelObject();
+		if (TargetLevel)
+		{
+			FIntPoint RTSize = ThumbnailRenderTarget->GetSizeXY();
+			
+			// Set persistent world package as transient to avoid package dirtying during thumbnail rendering
+			FUnmodifiableObject ImmuneWorld(TargetLevel->OwningWorld);
+			
+			FObjectThumbnail NewThumbnail;
+			// Generate the thumbnail
+			ThumbnailTools::RenderThumbnail(
+				TargetLevel,
+				RTSize.X,
+				RTSize.Y,
+				ThumbnailTools::EThumbnailTextureFlushMode::NeverFlush,
+				ThumbnailRenderTarget,
+				&NewThumbnail
+				);
+
+			UPackage* MyOutermostPackage = CastChecked<UPackage>(TargetLevel->GetOutermost());
+			ThumbnailTools::CacheThumbnail(LevelModel->GetAssetName().ToString(), &NewThumbnail, MyOutermostPackage);
+			UpdateTextureData(&NewThumbnail);
+		}
+	}
+}
+
+//----------------------------------------------------------------
+//
+//
+//----------------------------------------------------------------
+SWorldTileItem::~SWorldTileItem()
+{
+	LevelModel->ChangedEvent.RemoveAll(this);
+}
+
+void SWorldTileItem::Construct(const FArguments& InArgs)
+{
+	WorldModel = InArgs._InWorldModel;
+	LevelModel = InArgs._InItemModel;
+
+	ProgressBarImage = FEditorStyle::GetBrush(TEXT("ProgressBar.Marquee"));
+	
+	bNeedRefresh = true;
+	bIsDragging = false;
+	bAffectedByMarquee = false;
+
+	Thumbnail = MakeShareable(new FTileItemThumbnail(InArgs._ThumbnailRenderTarget, LevelModel));
+	ThumbnailViewport = SNew(SViewport)
+							.EnableGammaCorrection(false);
+							
+	ThumbnailViewport->SetViewportInterface(Thumbnail.ToSharedRef());
+	// This will grey out tile in case level is not loaded or locked
+	ThumbnailViewport->SetEnabled(TAttribute<bool>(this, &SWorldTileItem::IsItemEnabled));
+
+	LevelModel->ChangedEvent.AddSP(this, &SWorldTileItem::RequestRefresh);
+			
+	ChildSlot
+	[
+		ThumbnailViewport.ToSharedRef()
+	];
+
+	SetToolTip(CreateToolTipWidget());
+
+	CurveSequence = FCurveSequence(0.0f, 0.5f);
+	CurveSequence.Play();
+}
+
+void SWorldTileItem::RequestRefresh()
+{
+	bNeedRefresh = true;
+}
+
+UObject* SWorldTileItem::GetObjectBeingDisplayed() const
+{
+	return LevelModel->GetNodeObject();
+}
+
+BEGIN_SLATE_FUNCTION_BUILD_OPTIMIZATION
+TSharedRef<SToolTip> SWorldTileItem::CreateToolTipWidget()
+{
+	TSharedPtr<SToolTip>		TooltipWidget;
+	TSharedPtr<SVerticalBox>	TooltipSectionsWidget;
+	
+	SAssignNew(TooltipWidget, SToolTip)
+	.TextMargin(1)
+	.BorderImage(FEditorStyle::GetBrush("ContentBrowser.TileViewTooltip.ToolTipBorder"))
+	[
+		SNew(SBorder)
+		.Padding(6)
+		.BorderImage(FEditorStyle::GetBrush("ContentBrowser.TileViewTooltip.NonContentBorder"))
+		[
+			SAssignNew(TooltipSectionsWidget, SVerticalBox)
+
+			// Level name section
+			+SVerticalBox::Slot()
+			.AutoHeight()
+			.Padding(0,0,0,4)
+			[
+				SNew(SBorder)
+				.Padding(6)
+				.BorderImage(FEditorStyle::GetBrush("ContentBrowser.TileViewTooltip.ContentBorder"))
+				[
+					SNew(SVerticalBox)
+					
+					+SVerticalBox::Slot()
+					.VAlign(VAlign_Center)
+					.AutoHeight()
+					[
+						SNew(SHorizontalBox)
+
+						// Level name
+						+SHorizontalBox::Slot()
+						.Padding(6,0,0,0)
+						.HAlign(HAlign_Left)
+						.AutoWidth()
+						[
+							SNew(STextBlock)
+							.Text(this, &SWorldTileItem::GetLevelNameText) 
+							.Font(FEditorStyle::GetFontStyle("ContentBrowser.TileViewTooltip.NameFont"))
+						]
+					]
+				]
+			]
+		]
+	];
+
+	TArray< TPair<TAttribute<FText>, TAttribute<FText>> > CustomFields;
+	LevelModel->GetGridItemTooltipFields(CustomFields);
+
+	if (CustomFields.Num() > 0)
+	{
+		TSharedPtr<SUniformGridPanel> TooltipFieldsWidget;
+		// Add section for custom fields
+		TooltipSectionsWidget->AddSlot()
+			.AutoHeight()
+			.Padding(0,0,0,4)
+		[
+				SNew(SBorder)
+				.Padding(6)
+				.BorderImage(FEditorStyle::GetBrush("ContentBrowser.TileViewTooltip.ContentBorder"))
+				[
+					SAssignNew(TooltipFieldsWidget, SUniformGridPanel)
+				]
+		];
+
+		// Add custom fields
+		for (int32 FieldIdx = 0; FieldIdx < CustomFields.Num(); ++FieldIdx)
+		{
+			// Name:
+			TooltipFieldsWidget->AddSlot(0, FieldIdx)
+				.HAlign(HAlign_Left)
+			[
+				SNew(STextBlock)
+				.Text(CustomFields[FieldIdx].Key) 
+			];
+		
+			// Value:
+			TooltipFieldsWidget->AddSlot(1, FieldIdx)
+				.HAlign(HAlign_Right)
+			[
+				SNew(STextBlock)
+				.Text(CustomFields[FieldIdx].Value) 
+			];
+		}
+	}
+
+	return TooltipWidget.ToSharedRef();
+}
+END_SLATE_FUNCTION_BUILD_OPTIMIZATION
+
+FVector2D SWorldTileItem::GetPosition() const
+{
+	return LevelModel->GetLevelPosition2D();
+}
+
+const TSharedPtr<FLevelModel>& SWorldTileItem::GetLevelModel() const
+{
+	return LevelModel;
+}
+
+const FSlateBrush* SWorldTileItem::GetShadowBrush(bool bSelected) const
+{
+	return bSelected ? FEditorStyle::GetBrush(TEXT("Graph.CompactNode.ShadowSelected")) : FEditorStyle::GetBrush(TEXT("Graph.Node.Shadow"));
+}
+
+FOptionalSize SWorldTileItem::GetItemWidth() const
+{
+	return FOptionalSize(LevelModel->GetLevelSize2D().X);
+}
+
+FOptionalSize SWorldTileItem::GetItemHeight() const
+{
+	return FOptionalSize(LevelModel->GetLevelSize2D().Y);
+}
+
+FSlateRect SWorldTileItem::GetItemRect() const
+{
+	FVector2D LevelSize = LevelModel->GetLevelSize2D();
+	FVector2D LevelPos = GetPosition();
+	return FSlateRect(LevelPos, LevelPos + LevelSize);
+}
+
+TSharedPtr<IToolTip> SWorldTileItem::GetToolTip()
+{
+	// Hide tooltip in case item is being dragged now
+	if (LevelModel->GetLevelTranslationDelta().Size() > KINDA_SMALL_NUMBER)
+	{
+		return NULL;
+	}
+	
+	return SNodePanel::SNode::GetToolTip();
+}
+
+FVector2D SWorldTileItem::GetDesiredSizeForMarquee() const
+{
+	// we don't want to select items in non visible layers
+	if (!WorldModel->PassesAllFilters(LevelModel))
+	{
+		return FVector2D::ZeroVector;
+	}
+
+	return SNodePanel::SNode::GetDesiredSizeForMarquee();
+}
+
+FVector2D SWorldTileItem::ComputeDesiredSize() const
+{
+	return LevelModel->GetLevelSize2D();
+}
+
+int32 SWorldTileItem::OnPaint(const FGeometry& AllottedGeometry, const FSlateRect& ClippingRect, 
+								FSlateWindowElementList& OutDrawElements, int32 LayerId, 
+								const FWidgetStyle& InWidgetStyle, bool bParentEnabled) const
+{
+	const bool bIsVisible = FSlateRect::DoRectanglesIntersect(AllottedGeometry.GetClippingRect(), ClippingRect);
+	
+	if (bIsVisible)
+	{
+		// Redraw thumbnail image if requested
+		if (bNeedRefresh)
+		{
+			Thumbnail->UpdateThumbnail();
+			bNeedRefresh = false;
+		}
+		
+		LayerId = SNodePanel::SNode::OnPaint(AllottedGeometry, ClippingRect, OutDrawElements, LayerId, InWidgetStyle, bParentEnabled);
+		
+		const bool bSelected = (IsItemSelected() || bAffectedByMarquee);
+
+		// Draw the node's selection.
+		if (bSelected)
+		{
+			FPaintGeometry SelectionGeometry = AllottedGeometry.ToInflatedPaintGeometry(FVector2D(4,4)/AllottedGeometry.Scale);
+			SelectionGeometry.DrawScale = 0.5;
+				
+			FSlateDrawElement::MakeBox(
+				OutDrawElements,
+				LayerId + 1,
+				SelectionGeometry,
+				GetShadowBrush(bSelected),
+				ClippingRect
+				);
+		}
+
+		// Highlight potentially visible streaming levels from current preview point
+		//const auto& PreviewStreamingLevels = WorldModel->GetPreviewStreamingLevels();
+		//if (PreviewStreamingLevels.Find(LevelModel->TileDetails->PackageName) != INDEX_NONE)
+		//{
+		//	FPaintGeometry SelectionGeometry = AllottedGeometry.ToInflatedPaintGeometry(FVector2D(4,4)/AllottedGeometry.Scale);
+		//	SelectionGeometry.DrawScale = 0.5;
+
+		//	FSlateDrawElement::MakeBox(
+		//		OutDrawElements,
+		//		LayerId + 1,
+		//		SelectionGeometry,
+		//		GetShadowBrush(true),
+		//		ClippingRect,
+		//		ESlateDrawEffect::None,
+		//		FLinearColor::Green
+		//		);
+		//}
+
+		// Draw progress bar if level is currently loading
+		if (LevelModel->IsLoading())
+		{
+			const float ProgressBarAnimOffset = ProgressBarImage->ImageSize.X * CurveSequence.GetLerpLooping() / AllottedGeometry.Scale;
+			const float ProgressBarImageSize = ProgressBarImage->ImageSize.X / AllottedGeometry.Scale;
+			const float ProgressBarImageHeight = ProgressBarImage->ImageSize.Y / AllottedGeometry.Scale;
+			const FSlateRect ProgressBarClippingRect = AllottedGeometry.ToPaintGeometry().ToSlateRect().IntersectionWith(ClippingRect);
+		
+			FPaintGeometry LoadingBarPaintGeometry = AllottedGeometry.ToPaintGeometry(
+					FVector2D(ProgressBarAnimOffset - ProgressBarImageSize, 0 ),
+					FVector2D(AllottedGeometry.Size.X + ProgressBarImageSize, FMath::Min(AllottedGeometry.Size.Y, ProgressBarImageHeight))
+			);
+			LoadingBarPaintGeometry.DrawScale = 1.f;
+								
+			FSlateDrawElement::MakeBox(
+				OutDrawElements,
+				LayerId + 1,
+				LoadingBarPaintGeometry,
+				ProgressBarImage,
+				ProgressBarClippingRect,
+				ESlateDrawEffect::None,
+				FLinearColor(1.0f, 1.0f, 1.0f, 0.5f)
+			);
+		}
+	}
+	
+	return LayerId;
+}
+
+FReply SWorldTileItem::OnMouseButtonDoubleClick(const FGeometry& InMyGeometry, const FPointerEvent& InMouseEvent)
+{
+	if (OnHitTest(InMyGeometry, InMouseEvent.GetScreenSpacePosition()))
+	{
+		LevelModel->MakeLevelCurrent();
+		return FReply::Handled();
+	}
+	
+	return FReply::Unhandled();
+}
+
+bool SWorldTileItem::OnHitTest(const FGeometry& MyGeometry, FVector2D InAbsoluteCursorPosition)
+{
+	FVector2D CursorOffset = (InAbsoluteCursorPosition - MyGeometry.AbsolutePosition)/MyGeometry.Scale;
+	FVector2D LevelPos = LevelModel->GetLevelPosition2D();
+	FVector2D CursorWorldPos = LevelPos + CursorOffset;
+	return LevelModel->HitTest2D(CursorWorldPos);
+}
+
+FString SWorldTileItem::GetLevelNameText() const
+{
+	return LevelModel->GetDisplayName();
+}
+
+bool SWorldTileItem::IsItemEditable() const
+{
+	return LevelModel->IsEditable();
+}
+
+bool SWorldTileItem::IsItemSelected() const
+{
+	return LevelModel->GetLevelSelectionFlag();
+}
+
+bool SWorldTileItem::IsItemEnabled() const
+{
+	if (WorldModel->IsSimulating())
+	{
+		return LevelModel->IsVisible();
+	}
+	else
+	{
+		return LevelModel->IsEditable();
+	}
+}
+
+#undef LOCTEXT_NAMESPACE
diff --git a/Engine/Source/Editor/WorldBrowser/Private/Tiles/SWorldTileItem.h b/Engine/Source/Editor/WorldBrowser/Private/Tiles/SWorldTileItem.h
new file mode 100644
--- /dev/null
+++ b/Engine/Source/Editor/WorldBrowser/Private/Tiles/SWorldTileItem.h
@@ -0,0 +1,126 @@
+// Copyright 1998-2014 Epic Games, Inc. All Rights Reserved.
+#pragma once
+
+#include "SlateTextures.h"
+#include "GraphEditor.h"
+#include "SNodePanel.h"
+#include "AssetThumbnail.h"
+
+class FTileItemThumbnail : public ISlateViewport, public TSharedFromThis<FTileItemThumbnail>
+{
+public:
+	FTileItemThumbnail(FSlateTextureRenderTarget2DResource* InThumbnailRenderTarget, TSharedPtr<FLevelModel> InItemModel);
+	~FTileItemThumbnail();
+
+	/* ISlateViewport interface */
+	virtual FIntPoint GetSize() const OVERRIDE;
+	virtual class FSlateShaderResource* GetViewportRenderTargetTexture() const OVERRIDE;
+	virtual bool RequiresVsync() const OVERRIDE;
+	
+	/** Request thumbnail redraw*/
+	void UpdateThumbnail();
+
+private:
+	void UpdateTextureData(FObjectThumbnail* ObjectThumbnail);
+
+private:
+	TSharedPtr<FLevelModel>					LevelModel;
+	/** Rendering resource for slate */
+	FSlateTexture2DRHIRef*					ThumbnailTexture;
+	/** Shared render target for slate */
+	FSlateTextureRenderTarget2DResource*	ThumbnailRenderTarget;
+};
+
+//----------------------------------------------------------------
+//
+//
+//----------------------------------------------------------------
+class SWorldTileItem 
+	: public SNodePanel::SNode
+{
+public:
+	SLATE_BEGIN_ARGS(SWorldTileItem)
+	{}
+		/** The world data */
+		SLATE_ARGUMENT(TSharedPtr<FLevelCollectionModel>, InWorldModel)
+		/** Data for the asset this item represents */
+		SLATE_ARGUMENT(TSharedPtr<FLevelModel>, InItemModel)
+		//
+		SLATE_ARGUMENT(FSlateTextureRenderTarget2DResource*, ThumbnailRenderTarget)
+	SLATE_END_ARGS()
+
+	~SWorldTileItem();
+
+	void Construct(const FArguments& InArgs);
+	
+	// SNodePanel::SNode interface start
+	virtual FVector2D GetDesiredSizeForMarquee() const OVERRIDE;
+	virtual UObject* GetObjectBeingDisplayed() const OVERRIDE;
+	virtual FVector2D GetPosition() const OVERRIDE;
+	virtual const FSlateBrush* GetShadowBrush(bool bSelected) const OVERRIDE;
+	// SNodePanel::SNode interface end
+	
+	/** @return Deferred item refresh */
+	void RequestRefresh();
+
+	/** @return LevelModel associated with this item */
+	const TSharedPtr<FLevelModel>&	GetLevelModel() const;
+	
+	/** @return Item width in world units */
+	FOptionalSize GetItemWidth() const;
+	
+	/** @return Item height in world units */
+	FOptionalSize GetItemHeight() const;
+	
+	/** @return Rectangle in world units for this item as FSlateRect*/
+	FSlateRect GetItemRect() const;
+	
+	/** @return Whether this item can be edited (loaded and not locked) */
+	bool IsItemEditable() const;
+
+	/** @return Whether this item is selected */
+	bool IsItemSelected() const;
+	
+	/** @return Whether this item is enabled */
+	bool IsItemEnabled() const;
+
+private:
+	// SWidget interface start
+	virtual TSharedPtr<IToolTip> GetToolTip() OVERRIDE;
+	virtual FVector2D ComputeDesiredSize() const OVERRIDE;
+	virtual bool OnHitTest(const FGeometry& MyGeometry, FVector2D InAbsoluteCursorPosition) OVERRIDE;
+	virtual int32 OnPaint(const FGeometry& AllottedGeometry, const FSlateRect& MyClippingRect, 
+					FSlateWindowElementList& OutDrawElements, int32 LayerId, 
+					const FWidgetStyle& InWidgetStyle, bool bParentEnabled) const OVERRIDE;
+	virtual FReply OnMouseButtonDoubleClick(const FGeometry& InMyGeometry, const FPointerEvent& InMouseEvent) OVERRIDE;
+	// SWidget interface end
+
+	TSharedRef<SToolTip> CreateToolTipWidget();
+	
+	//
+	FString GetLevelNameText() const;
+
+public:
+	bool							bAffectedByMarquee;
+
+private:
+	/** The world data */
+	TSharedPtr<FLevelCollectionModel> WorldModel;
+
+	/** The data for this item */
+	TSharedPtr<FLevelModel>			LevelModel;
+
+	mutable bool					bNeedRefresh;
+	bool							bIsDragging;
+
+	TSharedPtr<FTileItemThumbnail>	Thumbnail;
+	/** The actual widget for the thumbnail. */
+	TSharedPtr<SViewport>			ThumbnailViewport;
+	
+	/** */
+	const FSlateBrush*				ProgressBarImage;
+
+	/** Curve sequence to drive loading animation */
+	FCurveSequence					CurveSequence;
+};
+
diff --git a/Engine/Source/Editor/WorldBrowser/Private/Tiles/WorldTileCollectionModel.cpp b/Engine/Source/Editor/WorldBrowser/Private/Tiles/WorldTileCollectionModel.cpp
--- a/Engine/Source/Editor/WorldBrowser/Private/Tiles/WorldTileCollectionModel.cpp
+++ b/Engine/Source/Editor/WorldBrowser/Private/Tiles/WorldTileCollectionModel.cpp
@@ -581,7 +581,7 @@ void FWorldTileCollectionModel::OnLevelsCollectionChanged()
 	if (WorldComposition)
 	{
 		// Force rescanning world composition tiles
-		WorldComposition->OpenWorldRoot(WorldComposition->GetWorldRoot());
+		WorldComposition->Rescan();
 		
 		// Initialize root level
 		TSharedPtr<FWorldTileModel> Item = AddLevelFromTile(INDEX_NONE);
diff --git a/Engine/Source/Editor/WorldBrowser/Private/Tiles/WorldTileCollectionModel.h b/Engine/Source/Editor/WorldBrowser/Private/Tiles/WorldTileCollectionModel.h
--- a/Engine/Source/Editor/WorldBrowser/Private/Tiles/WorldTileCollectionModel.h
+++ b/Engine/Source/Editor/WorldBrowser/Private/Tiles/WorldTileCollectionModel.h
@@ -38,6 +38,7 @@ public:
 	virtual bool CompareLevelsZOrder(TSharedPtr<FLevelModel> InA, TSharedPtr<FLevelModel> InB) const OVERRIDE;
 	virtual void RegisterDetailsCustomization(class FPropertyEditorModule& PropertyModule, TSharedPtr<class IDetailsView> InDetailsView)  OVERRIDE;
 	virtual void UnregisterDetailsCustomization(class FPropertyEditorModule& PropertyModule, TSharedPtr<class IDetailsView> InDetailsView) OVERRIDE;
+	virtual bool IsTileWorld() const OVERRIDE { return true; };
 	/** FLevelCollection interface end */
 
 private:
diff --git a/Engine/Source/Editor/WorldBrowser/Private/Tiles/WorldTileDetailsCustomization.cpp b/Engine/Source/Editor/WorldBrowser/Private/Tiles/WorldTileDetailsCustomization.cpp
--- a/Engine/Source/Editor/WorldBrowser/Private/Tiles/WorldTileDetailsCustomization.cpp
+++ b/Engine/Source/Editor/WorldBrowser/Private/Tiles/WorldTileDetailsCustomization.cpp
@@ -5,7 +5,7 @@
 #include "Editor/PropertyEditor/Public/PropertyEditing.h"
 #include "Editor/PropertyEditor/Public/PropertyEditorModule.h"
 
-#include "SWorldDetailsView.h"
+#include "SWorldDetails.h"
 #include "SPropertyEditorLevelPackage.h"
 #include "WorldTileDetails.h"
 #include "WorldTileCollectionModel.h"
diff --git a/Engine/Source/Editor/WorldBrowser/Private/WorldBrowserModule.cpp b/Engine/Source/Editor/WorldBrowser/Private/WorldBrowserModule.cpp
--- a/Engine/Source/Editor/WorldBrowser/Private/WorldBrowserModule.cpp
+++ b/Engine/Source/Editor/WorldBrowser/Private/WorldBrowserModule.cpp
@@ -2,13 +2,18 @@
 #include "WorldBrowserPrivatePCH.h"
 
 #include "ModuleManager.h"
-#include "SWorldMainView.h"
 #include "StreamingLevels/StreamingLevelEdMode.h"
-
+#include "Tiles/WorldTileCollectionModel.h"
+#include "StreamingLevels/StreamingLevelCollectionModel.h"
+#include "SWorldHierarchy.h"
+#include "SWorldDetails.h"
+#include "Tiles/SWorldComposition.h"
 
 
 IMPLEMENT_MODULE( FWorldBrowserModule, WorldBrowser );
 
+#define LOCTEXT_NAMESPACE "WorldBrowser"
+
 void FWorldBrowserModule::StartupModule()
 {
 	FLevelCollectionCommands::Register();
@@ -17,10 +22,20 @@ void FWorldBrowserModule::StartupModule()
 	TSharedRef<FStreamingLevelEdMode> NewEditorMode = MakeShareable(new FStreamingLevelEdMode);
 	GEditorModeTools().RegisterMode(NewEditorMode);
 	EdModeStreamingLevel = NewEditorMode;
+
+	FWorldDelegates::OnPostWorldInitialization.AddRaw(this, &FWorldBrowserModule::OnWorldCreated);
+	FWorldDelegates::OnWorldCleanup.AddRaw(this, &FWorldBrowserModule::OnWorldCleanup);
+	UWorldComposition::OnWorldCompositionCreated.AddRaw(this, &FWorldBrowserModule::OnWorldCompositionChanged);
+	UWorldComposition::OnWorldCompositionDestroyed.AddRaw(this, &FWorldBrowserModule::OnWorldCompositionChanged);
 }
 
 void FWorldBrowserModule::ShutdownModule()
 {
+	FWorldDelegates::OnPostWorldInitialization.RemoveAll(this);
+	FWorldDelegates::OnWorldCleanup.RemoveAll(this);
+	UWorldComposition::OnWorldCompositionCreated.RemoveAll(this);
+	UWorldComposition::OnWorldCompositionDestroyed.RemoveAll(this);
+
 	FLevelCollectionCommands::Unregister();
 
 	// unregister the editor mode
@@ -28,9 +43,89 @@ void FWorldBrowserModule::ShutdownModule()
 	EdModeStreamingLevel = NULL;
 }
 
-TSharedRef<SWidget> FWorldBrowserModule::CreateWorldBrowser()
+TSharedRef<SWidget> FWorldBrowserModule::CreateWorldBrowserHierarchy()
+{
+	UWorld* EditorWorld = GEditor->GetEditorWorldContext().World();
+	return SNew(SWorldHierarchy).InWorld(EditorWorld);
+}
+	
+TSharedRef<SWidget> FWorldBrowserModule::CreateWorldBrowserDetails()
+{
+	UWorld* EditorWorld = GEditor->GetEditorWorldContext().World();
+	return SNew(SWorldDetails).InWorld(EditorWorld);
+}
+
+TSharedRef<SWidget> FWorldBrowserModule::CreateWorldBrowserComposition()
+{
+	UWorld* EditorWorld = GEditor->GetEditorWorldContext().World();
+	return SNew(SWorldComposition).InWorld(EditorWorld);
+}
+
+void FWorldBrowserModule::OnWorldCreated(UWorld* InWorld, const UWorld::InitializationValues IVS)
 {
-	return SNew(SWorldMainView);
+	if (InWorld && 
+		InWorld->WorldType == EWorldType::Editor)
+	{
+		OnBrowseWorld.Broadcast(InWorld);
+	}
 }
 
+void FWorldBrowserModule::OnWorldCompositionChanged(UWorld* InWorld)
+{
+	if (InWorld && 
+		InWorld->WorldType == EWorldType::Editor)
+	{
+		OnBrowseWorld.Broadcast(NULL);
+		OnBrowseWorld.Broadcast(InWorld);
+	}
+}
+
+void FWorldBrowserModule::OnWorldCleanup(UWorld* InWorld, bool bSessionEnded, bool bCleanupResources)
+{
+	if (bCleanupResources)
+	{
+		TSharedPtr<FLevelCollectionModel> SharedWorldModel = WorldModel.Pin();
+		// Is there any editors alive?
+		if (SharedWorldModel.IsValid())
+		{
+			UWorld* ManagedWorld = SharedWorldModel->GetWorld(/*bEvenIfPendingKill*/true);
+			// Is it our world gets cleaned up?
+			if (ManagedWorld == InWorld)
+			{
+				// Will reset all references to a shared world model
+				OnBrowseWorld.Broadcast(NULL);
+				// So we have to be the last owner of this model
+				check(SharedWorldModel.IsUnique());
+			}
+		}
+	}
+}
+
+TSharedPtr<FLevelCollectionModel> FWorldBrowserModule::SharedWorldModel(UWorld* InWorld)
+{
+	TSharedPtr<FLevelCollectionModel> SharedWorldModel = WorldModel.Pin();
+	if (!SharedWorldModel.IsValid() || SharedWorldModel->GetWorld() != InWorld)
+	{
+		if (InWorld)
+		{
+			if (InWorld->WorldComposition)
+			{
+				SharedWorldModel = FWorldTileCollectionModel::Create(GEditor);
+			}
+			else
+			{
+				SharedWorldModel = FStreamingLevelCollectionModel::Create(GEditor);
+			}
+		}
+
+		// Hold weak reference to shared world model
+		WorldModel = SharedWorldModel;
+	}
+	
+	return SharedWorldModel;
+}
+
+
+
 
+#undef LOCTEXT_NAMESPACE
\ No newline at end of file
diff --git a/Engine/Source/Editor/WorldBrowser/Public/WorldBrowserModule.h b/Engine/Source/Editor/WorldBrowser/Public/WorldBrowserModule.h
--- a/Engine/Source/Editor/WorldBrowser/Public/WorldBrowserModule.h
+++ b/Engine/Source/Editor/WorldBrowser/Public/WorldBrowserModule.h
@@ -4,7 +4,7 @@
 #include "ModuleInterface.h"
 
 /**
- * The module holding all of the UI related pieces for Levels
+ * The module holding all of the UI related pieces for SubLevels management
  */
 class FWorldBrowserModule : public IModuleInterface
 {
@@ -20,12 +20,39 @@ public:
 	 * Called before the module is unloaded, right before the module object is destroyed.
 	 */
 	virtual void ShutdownModule();
+	
+	/**
+	 * Creates a levels hierarchy widget
+	 */
+	virtual TSharedRef<class SWidget> CreateWorldBrowserHierarchy();
+	
+	/**
+	 * Creates a levels details widget
+	 */
+	virtual TSharedRef<class SWidget> CreateWorldBrowserDetails();
 
 	/**
-	 * Creates a World Browser widget
+	 * Creates a levels composition widget
 	 */
-	virtual TSharedRef<class SWidget> CreateWorldBrowser();
+	virtual TSharedRef<class SWidget> CreateWorldBrowserComposition();
+
+	/**
+	 * @return world model shared between all World Browser editors
+	 */
+	virtual TSharedPtr<class FLevelCollectionModel> SharedWorldModel(UWorld* InWorld);
 	
+	/**
+	 * 
+	 */
+	DECLARE_EVENT_OneParam(FWorldBrowserModule, FOnBrowseWorld, UWorld*);
+	FOnBrowseWorld OnBrowseWorld;
+		
+private:
+	void OnWorldCreated(UWorld* InWorld, const UWorld::InitializationValues IVS);
+	void OnWorldCleanup(UWorld* InWorld, bool bSessionEnded, bool bCleanupResources);
+	void OnWorldCompositionChanged(UWorld* InWorld);
+		
 private:
+	TWeakPtr<class FLevelCollectionModel>	WorldModel;
 	TSharedPtr<class FStreamingLevelEdMode> EdModeStreamingLevel;
 };
diff --git a/Engine/Source/Runtime/Engine/Classes/Engine/World.h b/Engine/Source/Runtime/Engine/Classes/Engine/World.h
--- a/Engine/Source/Runtime/Engine/Classes/Engine/World.h
+++ b/Engine/Source/Runtime/Engine/Classes/Engine/World.h
@@ -857,7 +857,7 @@ public:
 	FIntPoint RequestedGlobalOriginOffset;
 	
 	/** All levels information from which our world is composed */
-	UPROPERTY(Transient)
+	UPROPERTY()
 	class UWorldComposition* WorldComposition;
 	
 	/** Streaming level package name to LOD index. 
@@ -1623,8 +1623,6 @@ public:
 	virtual void PostSaveRoot( bool bCleanupIsRequired ) OVERRIDE;
 	virtual UWorld* GetWorld() const OVERRIDE;
 	static void AddReferencedObjects(UObject* InThis, FReferenceCollector& Collector);
-	virtual void PostDuplicate(bool bDuplicateForPIE) OVERRIDE;
-
 	// End UObject Interface
 	
 	/**
@@ -1783,7 +1781,6 @@ public:
 			, bShouldSimulatePhysics(true)
 			, bEnableTraceCollision(false)
 			, bTransactional(true)
-			, bCreateWorldComposition(false)
 		{
 		}
 
@@ -1795,7 +1792,6 @@ public:
 		uint32 bShouldSimulatePhysics:1;
 		uint32 bEnableTraceCollision:1;
 		uint32 bTransactional:1;
-		uint32 bCreateWorldComposition:1;
 
 
 		InitializationValues& InitializeScenes(const bool bInitialize) { bInitializeScenes = bInitialize; return *this; }
@@ -1806,7 +1802,6 @@ public:
 		InitializationValues& ShouldSimulatePhysics(const bool bInShouldSimulatePhysics) { bShouldSimulatePhysics = bInShouldSimulatePhysics; return *this; }
 		InitializationValues& EnableTraceCollision(const bool bInEnableTraceCollision) { bEnableTraceCollision = bInEnableTraceCollision; return *this; }
 		InitializationValues& SetTransactional(const bool bInTransactional) { bTransactional = bInTransactional; return *this; }
-		InitializationValues& CreateWorldComposition(const bool bCreate) { bCreateWorldComposition = bCreate; return *this; }
 	};
 
 	/**
@@ -2391,9 +2386,6 @@ public:
 	/** Retrieves information whether all navigation with this world has been rebuilt */
 	bool IsNavigationRebuilt() const;
 
-	/** Setup runtime objects for world composition based on folder of currently loaded map */
-	void InitializeWorldComposition();
-
 	/** Request to translate world origin to specified position on next tick */
 	void RequestNewWorldOrigin(const FIntPoint& InNewOrigin);
 	
diff --git a/Engine/Source/Runtime/Engine/Classes/Engine/WorldComposition.h b/Engine/Source/Runtime/Engine/Classes/Engine/WorldComposition.h
--- a/Engine/Source/Runtime/Engine/Classes/Engine/WorldComposition.h
+++ b/Engine/Source/Runtime/Engine/Classes/Engine/WorldComposition.h
@@ -77,9 +77,6 @@ class ENGINE_API UWorldComposition : public UObject
 	/** Returns currently visible and hidden levels based on distance based streaming */
 	void GetDistanceVisibleLevels(const FVector& InLocation, TArray<FDistanceVisibleLevel>& OutVisibleLevels, TArray<FDistanceVisibleLevel>& OutHiddenLevels) const;
 
-	/** Opens world composition from specified folder (long PackageName)*/
-	bool OpenWorldRoot(const FString& InPathToRoot);
-	
 	/** @returns Currently opened world composition root folder (long PackageName)*/
 	FString GetWorldRoot() const;
 	
@@ -107,6 +104,9 @@ class ENGINE_API UWorldComposition : public UObject
 	/** @returns Level bounding box in current shifted space */
 	FBox GetLevelBounds(ULevel* InLevel) const;
 
+	/** Scans world root folder for relevant packages and initializes world composition structures */
+	void Rescan();
+
 #if WITH_EDITOR
 	/** @returns FWorldTileInfo associated with specified package */
 	FWorldTileInfo GetTileInfo(const FName& InPackageName) const;
@@ -121,19 +121,24 @@ class ENGINE_API UWorldComposition : public UObject
 	void RestoreDirtyTilesInfo(const FTilesList& TilesPrevState);
 	
 	/** Collect tiles package names to cook  */
-	static bool CollectTilesToCook(const FString& CmdLineMapEntry, TArray<FString>& FilesInPath);
+	void CollectTilesToCook(TArray<FString>& PackageNames);
+
+	DECLARE_MULTICAST_DELEGATE_OneParam(FWorldCompositionEvent, UWorld*);
+	// Callback on world composition creation
+	static FWorldCompositionEvent OnWorldCompositionCreated;
+	// Callback on world composition destruction 
+	static FWorldCompositionEvent OnWorldCompositionDestroyed;
 
 #endif //WITH_EDITOR
 
 private:
 	// UObject interface
-	/** Handles WorldComposition duplication for PIE */
-	virtual void Serialize( FArchive& Ar ) OVERRIDE;
+	virtual void PostInitProperties() OVERRIDE;
+	virtual void Serialize(FArchive& Ar) OVERRIDE;
+	virtual void PostDuplicate(bool bDuplicateForPIE) OVERRIDE;
+	virtual void PostLoad() OVERRIDE;
 	// UObject interface
 
-	/** Scans world root folder for relevant packages and initializes world composition structures */
-	void Rescan();
-
 	/** Populate streaming level objects using tiles information */
 	void PopulateStreamingLevels();
 
@@ -173,7 +178,7 @@ private:
 
 public:
 	// Streaming level objects for each tile
-	UPROPERTY()
+	UPROPERTY(transient)
 	TArray<ULevelStreaming*>	TilesStreaming;
 
 	// Time threshold between tile streaming state changes
diff --git a/Engine/Source/Runtime/Engine/Classes/GameFramework/WorldSettings.h b/Engine/Source/Runtime/Engine/Classes/GameFramework/WorldSettings.h
--- a/Engine/Source/Runtime/Engine/Classes/GameFramework/WorldSettings.h
+++ b/Engine/Source/Runtime/Engine/Classes/GameFramework/WorldSettings.h
@@ -188,6 +188,14 @@ class ENGINE_API AWorldSettings : public AInfo
 	UPROPERTY(EditAnywhere, BlueprintReadOnly, Category=Config)
 	uint32 bEnableNavigationSystem:1;
 
+	/** if set to true enables tools for composing a tiled world */
+	UPROPERTY(EditAnywhere, BlueprintReadOnly, Category=Config)
+	uint32 bEnableWorldComposition:1;
+
+	/** if set to true world origin will shift to a camera position when camera goes far away from current origin */
+	UPROPERTY(EditAnywhere, BlueprintReadOnly, Category=Config, meta=(editcondition = "bEnableWorldComposition"))
+	uint32 bEnableWorldOriginRebasing:1;
+		
 	/** if set to true, when we call GetGravityZ we assume WorldGravityZ has already been initialized and skip the lookup of DefaultGravityZ and GlobalGravityZ */
 	UPROPERTY(transient)
 	uint32 bWorldGravitySet:1;
@@ -439,5 +447,8 @@ private:
 	HIDE_ACTOR_TRANSFORM_FUNCTIONS();
 
 	virtual void Serialize( FArchive& Ar ) OVERRIDE;
+
+	/** Toggles world composition mode */
+	void EnabledWorldComposition(bool bEnable);
 };
 
diff --git a/Engine/Source/Runtime/Engine/Private/UnrealEngine.cpp b/Engine/Source/Runtime/Engine/Private/UnrealEngine.cpp
--- a/Engine/Source/Runtime/Engine/Private/UnrealEngine.cpp
+++ b/Engine/Source/Runtime/Engine/Private/UnrealEngine.cpp
@@ -8168,9 +8168,6 @@ bool UEngine::LoadMap( FWorldContext& WorldContext, FURL URL, class UPendingNetG
 
 	MALLOC_PROFILER( FMallocProfiler::SnapshotMemoryLoadMapMid( URL.Map ); )
 
-	// Whether we should treat URL.Map as world root for WorldComposition
-	const bool bWorldComposition = URL.HasOption(TEXT("worldcomposition"));
-
 	if( GUseSeekFreeLoading )
 	{
 		// Load GameMode specific data
@@ -8343,7 +8340,7 @@ bool UEngine::LoadMap( FWorldContext& WorldContext, FURL URL, class UPendingNetG
 	// In the PIE case the world will already have been initialized as part of CreatePIEWorldByDuplication
 	if (!WorldContext.World()->bIsWorldInitialized)
 	{
-		WorldContext.World()->InitWorld(UWorld::InitializationValues().CreateWorldComposition(bWorldComposition));
+		WorldContext.World()->InitWorld();
 	}
 
 	// Handle pending level.
@@ -8390,7 +8387,7 @@ bool UEngine::LoadMap( FWorldContext& WorldContext, FURL URL, class UPendingNetG
 	LoadPackagesFully(WorldContext.World(), FULLYLOAD_Map, WorldContext.World()->PersistentLevel->GetOutermost()->GetName());
 
 	// Set initial world origin and stream in levels
-	if (bWorldComposition)
+	if (WorldContext.World()->WorldComposition)
 	{
 		WorldContext.World()->NavigateTo(FIntPoint::ZeroValue);
 	}
diff --git a/Engine/Source/Runtime/Engine/Private/World.cpp b/Engine/Source/Runtime/Engine/Private/World.cpp
--- a/Engine/Source/Runtime/Engine/Private/World.cpp
+++ b/Engine/Source/Runtime/Engine/Private/World.cpp
@@ -135,7 +135,6 @@ void UWorld::Serialize( FArchive& Ar )
 	// Serialize World composition for PIE
 	if (Ar.GetPortFlags() & PPF_DuplicateForPIE)
 	{
-		Ar << WorldComposition;
 		Ar << GlobalOriginOffset;
 		Ar << RequestedGlobalOriginOffset;
 	}
@@ -197,25 +196,6 @@ void UWorld::AddReferencedObjects(UObject* InThis, FReferenceCollector& Collecto
 	Super::AddReferencedObjects( This, Collector );
 }
 
-void UWorld::PostDuplicate(bool bDuplicateForPIE)
-{
-	if (bDuplicateForPIE)
-	{
-		if (WorldComposition)
-		{
-			// Remove streaming levels created by World Browser, to avoid duplication with streaming levels from world composition
-			auto Predicate = [](ULevelStreaming* StreamingLevel)
-			{
-				return (!StreamingLevel || StreamingLevel->HasAnyFlags(RF_Transient));
-			};
-			StreamingLevels.RemoveAll(Predicate);
-
-			// Add streaming levels managed by world composition
-			StreamingLevels.Append(WorldComposition->TilesStreaming);
-		}
-	}
-}
-
 void UWorld::FinishDestroy()
 {
 	// Avoid cleanup if the world hasn't been initialized, e.g., the default object or a world object that got loaded
@@ -762,11 +742,6 @@ void UWorld::InitWorld(const InitializationValues IVS)
 	// update it's bIsDefaultLevel
 	bIsDefaultLevel = (FPaths::GetBaseFilename(GetMapName()) == FPaths::GetBaseFilename(UGameMapsSettings::GetGameDefaultMap()));
 
-	if (IVS.bCreateWorldComposition)
-	{
-		InitializeWorldComposition();
-	}
-
 	// We're initialized now.
 	bIsWorldInitialized = true;
 
@@ -2173,7 +2148,7 @@ void UWorld::UpdateLevelStreaming( FSceneViewFamily* ViewFamily )
 
 void UWorld::EvaluateWorldOriginLocation(const FSceneViewFamily& ViewFamily)
 {
-	if ( IsGameWorld() )
+	if (IsGameWorld() && GetWorldSettings()->bEnableWorldOriginRebasing)
 	{
 		FVector CentroidLocation(0,0,0);
 		for (int32 ViewIndex = 0; ViewIndex < ViewFamily.Views.Num(); ViewIndex++)
@@ -3164,11 +3139,6 @@ void UWorld::WelcomePlayer(UNetConnection* Connection)
 	Connection->SendPackageMap();
 	
 	FString LevelName = CurrentLevel->GetOutermost()->GetName();
-	if (WorldComposition != NULL)
-	{
-		LevelName+= TEXT("?worldcomposition");
-	}
-
 	Connection->ClientWorldPackageName = CurrentLevel->GetOutermost()->GetFName();
 
 	FString GameName;
@@ -3686,28 +3656,13 @@ void UWorld::CancelPendingMapChange()
 
 void UWorld::CommitMapChange()
 {
-		if( IsPreparingMapChange() )
-		{
+	if( IsPreparingMapChange() )
+	{
 		GEngine->SetShouldCommitPendingMapChange(this, true);
-		}
-		else
-		{
-			UE_LOG(LogWorld, Log, TEXT("AWorldSettings::CommitMapChange being called without a pending map change!"));
-		}
 	}
-
-void UWorld::InitializeWorldComposition()
-{
-	WorldComposition = NULL;
-	//
-	FString RootPackageName = GetOutermost()->GetName();
-	if (FPackageName::DoesPackageExist(RootPackageName))
+	else
 	{
-		WorldComposition = ConstructObject<UWorldComposition>(UWorldComposition::StaticClass(), this);
-		// Get the root directory
-		FString RootPathName = FPaths::GetPath(RootPackageName) + TEXT("/");
-		// Open world composition from root directory
-		WorldComposition->OpenWorldRoot(RootPathName);
+		UE_LOG(LogWorld, Log, TEXT("AWorldSettings::CommitMapChange being called without a pending map change!"));
 	}
 }
 
diff --git a/Engine/Source/Runtime/Engine/Private/WorldComposition.cpp b/Engine/Source/Runtime/Engine/Private/WorldComposition.cpp
--- a/Engine/Source/Runtime/Engine/Private/WorldComposition.cpp
+++ b/Engine/Source/Runtime/Engine/Private/WorldComposition.cpp
@@ -6,26 +6,64 @@
 #include "EnginePrivate.h"
 #include "LevelUtils.h"
 
+#if WITH_EDITOR
+UWorldComposition::FWorldCompositionEvent UWorldComposition::OnWorldCompositionCreated;
+UWorldComposition::FWorldCompositionEvent UWorldComposition::OnWorldCompositionDestroyed;
+#endif // WITH_EDITOR
+
 UWorldComposition::UWorldComposition(const class FPostConstructInitializeProperties& PCIP)
 	: Super(PCIP)
 {
 }
 
+void UWorldComposition::PostInitProperties()
+{
+	Super::PostInitProperties();
+
+	if (!IsTemplate() && (GetOutermost()->PackageFlags & PKG_PlayInEditor) == 0)
+	{
+		// Tiles information is not serialized to disk, and should be regenerated on world composition object construction
+		Rescan();
+	}
+}
+
 void UWorldComposition::Serialize( FArchive& Ar )
 {
 	Super::Serialize(Ar);
-
-	Ar << WorldRoot;
-	Ar << Tiles;
-	Ar << TilesStreaming;
 	
-	if (Ar.IsLoading() && 
-		Ar.GetPortFlags() & PPF_DuplicateForPIE)
+	// We serialize this data only for PIE
+	// In normal game this data is regenerated on object construction
+	if (Ar.GetPortFlags() & PPF_DuplicateForPIE)
 	{
-		// Fixup tile names to PIE 
+		Ar << WorldRoot;
+		Ar << Tiles;
+		Ar << TilesStreaming;
+	}
+}
+
+void UWorldComposition::PostDuplicate(bool bDuplicateForPIE)
+{
+	Super::PostDuplicate(bDuplicateForPIE);
+
 #if WITH_EDITOR
+	if (bDuplicateForPIE)
+	{
 		FixupForPIE(GetOutermost()->PIEInstanceID);	
-#endif
+	}
+#endif // WITH_EDITOR
+}
+
+void UWorldComposition::PostLoad()
+{
+	Super::PostLoad();
+
+	if (GetWorld()->IsGameWorld())
+	{
+		// Remove streaming levels created by World Browser, to avoid duplication with streaming levels from world composition
+		GetWorld()->StreamingLevels.Empty();
+		
+		// Add streaming levels managed by world composition
+		GetWorld()->StreamingLevels.Append(TilesStreaming);
 	}
 }
 
@@ -34,7 +72,6 @@ void UWorldComposition::FixupForPIE(int32 PIEInstanceID)
 	for (int32 TileIdx = 0; TileIdx < Tiles.Num(); ++TileIdx)
 	{
 		FWorldCompositionTile& Tile = Tiles[TileIdx];
-		ULevelStreaming* TileStreaming = TilesStreaming[TileIdx];
 		
 		FString PIEPackageName = UWorld::ConvertToPIEPackageName(Tile.PackageName.ToString(), PIEInstanceID);
 		Tile.PackageName = FName(*PIEPackageName);
@@ -43,25 +80,7 @@ void UWorldComposition::FixupForPIE(int32 PIEInstanceID)
 			FString PIELODPackageName = UWorld::ConvertToPIEPackageName(LODPackageName.ToString(), PIEInstanceID);
 			LODPackageName = FName(*PIELODPackageName);
 		}
-
-		TileStreaming->PackageNameToLoad = TileStreaming->PackageName;
-		TileStreaming->PackageName = Tile.PackageName;
-	}
-}
-
-bool UWorldComposition::OpenWorldRoot(const FString& PathToRoot)
-{
-	WorldRoot = PathToRoot;
-	Rescan();
-
-	// Add streaming levels to our world when we are in game only
-	// In Editor we use own level streaming mechanism not related to ULevelStreaming
-	if (GetWorld()->IsGameWorld())
-	{
-		GetWorld()->StreamingLevels.Append(TilesStreaming);
 	}
-
-	return true;
 }
 
 FString UWorldComposition::GetWorldRoot() const
@@ -151,11 +170,14 @@ void UWorldComposition::Rescan()
 	FTilesList SavedTileList = Tiles;
 		
 	Reset();	
-	
-	if (WorldRoot.IsEmpty())
+
+	FString RootPackageName = GetOutermost()->GetName();
+	if (!FPackageName::DoesPackageExist(RootPackageName))
 	{
-		return;
+		return;	
 	}
+	
+	WorldRoot = FPaths::GetPath(RootPackageName) + TEXT("/");
 
 	UWorld* OwningWorld = GetWorld();
 		
@@ -255,6 +277,7 @@ void UWorldComposition::CaclulateTilesAbsolutePositions()
 
 void UWorldComposition::Reset()
 {
+	WorldRoot.Empty();
 	Tiles.Empty();
 	TilesStreaming.Empty();
 }
@@ -369,34 +392,17 @@ void UWorldComposition::RestoreDirtyTilesInfo(const FTilesList& TilesPrevState)
 	}
 }
 
-bool UWorldComposition::CollectTilesToCook(const FString& CmdLineMapEntry, TArray<FString>& FilesInPath)
+void UWorldComposition::CollectTilesToCook(TArray<FString>& PackageNames)
 {
-	static const FString WorldCompositionOption = TEXT("?worldcomposition");
-
-	if (CmdLineMapEntry.EndsWith(WorldCompositionOption))
+	for (const FWorldCompositionTile& Tile : Tiles)
 	{
-		// Chop option string from package name
-		FString RootPackageName = CmdLineMapEntry.LeftChop(WorldCompositionOption.Len());
-		FString RootFilename;
-		if (!FPackageName::SearchForPackageOnDisk(RootPackageName, NULL, &RootFilename))
-		{
-			return false;
-		}
-
-		// All maps inside root map folder and subfolders are participate in world composition
-		TArray<FString> Files;
-		FString RootFolder = FPaths::GetPath(RootFilename);
-		IFileManager::Get().FindFilesRecursive(Files, *RootFolder, *(FString(TEXT("*")) + FPackageName::GetMapPackageExtension()), true, false, false);
-
-		for (FString StdFile : Files)
+		PackageNames.AddUnique(Tile.PackageName.ToString());
+		
+		for (const FName& TileLODName : Tile.LODPackageNames)
 		{
-			FPaths::MakeStandardFilename(StdFile);
-			FilesInPath.AddUnique(StdFile);
+			PackageNames.AddUnique(TileLODName.ToString());
 		}
-		return true;
 	}
-	
-	return false;
 }
 
 #endif //WITH_EDITOR
diff --git a/Engine/Source/Runtime/Engine/Private/WorldSettings.cpp b/Engine/Source/Runtime/Engine/Private/WorldSettings.cpp
--- a/Engine/Source/Runtime/Engine/Private/WorldSettings.cpp
+++ b/Engine/Source/Runtime/Engine/Private/WorldSettings.cpp
@@ -26,6 +26,8 @@ AWorldSettings::AWorldSettings(const class FPostConstructInitializeProperties& P
 
 	bEnableWorldBoundsChecks = true;
 	bEnableNavigationSystem = true;
+	bEnableWorldComposition = false;
+	bEnableWorldOriginRebasing = true;
 
 	KillZ = -HALF_WORLD_MAX1;
 	KillZDamageType = ConstructorStatics.DmgType_Environmental_Object.Object;
@@ -197,6 +199,11 @@ void AWorldSettings::PostEditChangeProperty(FPropertyChangedEvent& PropertyChang
 		{
 			FMessageDialog::Open( EAppMsgType::Ok, LOCTEXT("bForceNoPrecomputedLightingIsEnabled", "bForceNoPrecomputedLighting is now enabled, build lighting once to propagate the change (will remove existing precomputed lighting data)."));
 		}
+
+		if (PropertyThatChanged->GetName()==TEXT("bEnableWorldComposition"))
+		{
+			EnabledWorldComposition(bEnableWorldComposition);
+		}
 	}
 
 	LightmassSettings.NumIndirectLightingBounces = FMath::Clamp(LightmassSettings.NumIndirectLightingBounces, 0, 100);
@@ -226,6 +233,38 @@ void AWorldSettings::PostEditChangeProperty(FPropertyChangedEvent& PropertyChang
 	Super::PostEditChangeProperty(PropertyChangedEvent);
 }
 
+void AWorldSettings::EnabledWorldComposition(bool bEnable)
+{
+	bEnableWorldComposition = false;
+
+	if (!bEnable)
+	{
+		if (GetWorld()->WorldComposition != nullptr)
+		{
+			GetWorld()->WorldComposition = nullptr;
+			GetWorld()->StreamingLevels.Empty();
+			GetWorld()->FlushLevelStreaming();
+
+			UWorldComposition::OnWorldCompositionDestroyed.Broadcast(GetWorld());
+		}
+	}
+	else
+	{
+		FString RootPackageName = GetOutermost()->GetName();
+
+		if (GetWorld()->WorldComposition == nullptr && FPackageName::DoesPackageExist(RootPackageName))
+		{
+			GetWorld()->StreamingLevels.Empty();
+			GetWorld()->FlushLevelStreaming();
+			GetWorld()->WorldComposition = ConstructObject<UWorldComposition>(UWorldComposition::StaticClass(), GetWorld());
+			
+			bEnableWorldComposition = true;
+
+			UWorldComposition::OnWorldCompositionCreated.Broadcast(GetWorld());
+		}
+	}
+}
+
 #endif // WITH_EDITOR
 
 #undef LOCTEXT_NAMESPACE
