Disabled 'focusing' in World Browser when world orign rebasing is disabled in World Settings

diff --git a/Engine/Source/Editor/WorldBrowser/Private/Tiles/WorldTileCollectionModel.cpp b/Engine/Source/Editor/WorldBrowser/Private/Tiles/WorldTileCollectionModel.cpp
--- a/Engine/Source/Editor/WorldBrowser/Private/Tiles/WorldTileCollectionModel.cpp
+++ b/Engine/Source/Editor/WorldBrowser/Private/Tiles/WorldTileCollectionModel.cpp
@@ -270,7 +270,7 @@ void FWorldTileCollectionModel::BuildGridMenu(FMenuBuilder& MenuBuilder) const
 	
 	const FLevelCollectionCommands& Commands = FLevelCollectionCommands::Get();
 	
-	if (!AreAnyLevelsSelected())
+	if (!AreAnyLevelsSelected() && GetWorld()->GetWorldSettings()->bEnableWorldOriginRebasing)
 	{
 		MenuBuilder.AddMenuEntry(Commands.ResetWorldOrigin);
 		return;
@@ -300,11 +300,14 @@ void FWorldTileCollectionModel::BuildGridMenu(FMenuBuilder& MenuBuilder) const
 		}
 		MenuBuilder.EndSection();
 
-		MenuBuilder.BeginSection("WorldGridItemsWorldOrigin");
+		if (GetWorld()->GetWorldSettings()->bEnableWorldOriginRebasing)
 		{
-			MenuBuilder.AddMenuEntry(Commands.MoveWorldOrigin);
+			MenuBuilder.BeginSection("WorldGridItemsWorldOrigin");
+			{
+				MenuBuilder.AddMenuEntry(Commands.MoveWorldOrigin);
+			}
+			MenuBuilder.EndSection();
 		}
-		MenuBuilder.EndSection();
 
 		// Landscape specific commands
 		if (StaticCastSharedPtr<FWorldTileModel>(SelectedLevelsList[0])->IsLandscapeBased())
@@ -842,7 +845,7 @@ TSharedPtr<FLevelModel> FWorldTileCollectionModel::CreateNewEmptyLevel()
 
 void FWorldTileCollectionModel::Focus(FBox InArea, FocusStrategy InStrategy)
 {
-	if (IsReadOnly() || !InArea.IsValid)
+	if (IsReadOnly() || !InArea.IsValid || !GetWorld()->GetWorldSettings()->bEnableWorldOriginRebasing)
 	{
 		return;
 	}
@@ -1457,7 +1460,7 @@ void FWorldTileCollectionModel::MoveWorldOrigin(const FIntPoint& InOrigin)
 
 void FWorldTileCollectionModel::MoveWorldOrigin_Executed()
 {
-	if (!IsOneLevelSelected())
+	if (!IsOneLevelSelected() || !GetWorld()->GetWorldSettings()->bEnableWorldOriginRebasing)
 	{
 		return;
 	}
@@ -1468,10 +1471,13 @@ void FWorldTileCollectionModel::MoveWorldOrigin_Executed()
 
 void FWorldTileCollectionModel::ResetWorldOrigin_Executed()
 {
-	FBox OriginArea = EditableWorldArea().ShiftBy(FVector(GetWorld()->GlobalOriginOffset, 0));
-	Focus(OriginArea, OriginAtCenter);
+	if (GetWorld()->GetWorldSettings()->bEnableWorldOriginRebasing)
+	{
+		FBox OriginArea = EditableWorldArea().ShiftBy(FVector(GetWorld()->GlobalOriginOffset, 0));
+		Focus(OriginArea, OriginAtCenter);
 	
-	MoveWorldOrigin(FIntPoint::ZeroValue);
+		MoveWorldOrigin(FIntPoint::ZeroValue);
+	}
 }
 
 void FWorldTileCollectionModel::ResetLevelOrigin_Executed()
diff --git a/Engine/Source/Editor/WorldBrowser/Private/Tiles/WorldTileModel.cpp b/Engine/Source/Editor/WorldBrowser/Private/Tiles/WorldTileModel.cpp
--- a/Engine/Source/Editor/WorldBrowser/Private/Tiles/WorldTileModel.cpp
+++ b/Engine/Source/Editor/WorldBrowser/Private/Tiles/WorldTileModel.cpp
@@ -229,6 +229,12 @@ bool FWorldTileModel::ShouldBeVisible(FBox EditableArea) const
 		return true;
 	}
 
+	// Visibility does not depend on level positions when world origin rebasing is disabled
+	if (!LevelCollectionModel.GetWorld()->GetWorldSettings()->bEnableWorldOriginRebasing)
+	{
+		return true;
+	}
+
 	FBox LevelBBox = GetLevelBounds();
 
 	// Visible if level has no valid bounds
@@ -303,7 +309,7 @@ bool FWorldTileModel::IsShelved() const
 
 void FWorldTileModel::Shelve()
 {
-	if (LevelCollectionModel.IsReadOnly() || IsShelved() || IsRootTile())
+	if (LevelCollectionModel.IsReadOnly() || IsShelved() || IsRootTile() || !LevelCollectionModel.GetWorld()->GetWorldSettings()->bEnableWorldOriginRebasing)
 	{
 		return;
 	}
