Changed colors of world axis in the World Browser to be similar with editor gizmo colors

diff --git a/Engine/Source/Editor/WorldBrowser/Private/Tiles/SWorldComposition.cpp b/Engine/Source/Editor/WorldBrowser/Private/Tiles/SWorldComposition.cpp
--- a/Engine/Source/Editor/WorldBrowser/Private/Tiles/SWorldComposition.cpp
+++ b/Engine/Source/Editor/WorldBrowser/Private/Tiles/SWorldComposition.cpp
@@ -298,6 +298,7 @@ public:
 			FPaintGeometry EditableArea = AllottedGeometry.ToPaintGeometry(PaintPosition, PaintSize);
 			EditableArea.DrawScale = 0.2f;
 			FLinearColor PaintColor = FLinearColor::Red;
+			PaintColor.A = 0.4f;
 
 			FSlateDrawElement::MakeBox(
 				OutDrawElements,
@@ -453,13 +454,16 @@ protected:
 		FVector2D ScreenWorldOrigin = GraphCoordToPanelCoord(FVector2D(0, 0));
 		FSlateRect ScreenRect = AllottedGeometry.GetRect();
 	
-		// Vertical line
+		// World Y-axis
 		if (ScreenWorldOrigin.X > ScreenRect.Left &&
 			ScreenWorldOrigin.X < ScreenRect.Right)
 		{
 			TArray<FVector2D> LinePoints;
 			LinePoints.Add(FVector2D(ScreenWorldOrigin.X, ScreenRect.Top));
 			LinePoints.Add(FVector2D(ScreenWorldOrigin.X, ScreenRect.Bottom));
+
+			FLinearColor YAxisColor = FLinearColor::Green;
+			YAxisColor.A = 0.4f;
 		
 			FSlateDrawElement::MakeLines( 
 				OutDrawElements,
@@ -468,16 +472,19 @@ protected:
 				LinePoints,
 				MyClippingRect,
 				ESlateDrawEffect::None,
-				FColor(100, 100, 100));
+				YAxisColor);
 		}
 
-		// Horizontal line
+		// World X-axis
 		if (ScreenWorldOrigin.Y > ScreenRect.Top &&
 			ScreenWorldOrigin.Y < ScreenRect.Bottom)
 		{
 			TArray<FVector2D> LinePoints;
 			LinePoints.Add(FVector2D(ScreenRect.Left, ScreenWorldOrigin.Y));
 			LinePoints.Add(FVector2D(ScreenRect.Right, ScreenWorldOrigin.Y));
+
+			FLinearColor XAxisColor = FLinearColor::Red;
+			XAxisColor.A = 0.4f;
 		
 			FSlateDrawElement::MakeLines( 
 				OutDrawElements,
@@ -486,7 +493,7 @@ protected:
 				LinePoints,
 				MyClippingRect,
 				ESlateDrawEffect::None,
-				FColor(100, 100, 100));
+				XAxisColor);
 		}
 
 		return LayerId + 1;
