Improvements to World Browser tiles tooltips

diff --git a/Engine/Source/Editor/WorldBrowser/Private/LevelModel.h b/Engine/Source/Editor/WorldBrowser/Private/LevelModel.h
--- a/Engine/Source/Editor/WorldBrowser/Private/LevelModel.h
+++ b/Engine/Source/Editor/WorldBrowser/Private/LevelModel.h
@@ -184,9 +184,6 @@ public:
 	/**	@return Whether it's possible to drop onto this level */
 	virtual bool IsGoodToDrop(const TSharedPtr<FLevelDragDropOp>& Op) const;
 
-	/**	@return Name->Value fields this level want to add to a grid item tooltip */
-	virtual void GetGridItemTooltipFields(TArray< TPair<TAttribute<FText>, TAttribute<FText>> >& CustomFields) const = 0;
-
 	/** Notification when level was added(shown) to world */
 	virtual void OnLevelAddedToWorld(ULevel* InLevel);
 
diff --git a/Engine/Source/Editor/WorldBrowser/Private/StreamingLevels/StreamingLevelModel.cpp b/Engine/Source/Editor/WorldBrowser/Private/StreamingLevels/StreamingLevelModel.cpp
--- a/Engine/Source/Editor/WorldBrowser/Private/StreamingLevels/StreamingLevelModel.cpp
+++ b/Engine/Source/Editor/WorldBrowser/Private/StreamingLevels/StreamingLevelModel.cpp
@@ -94,11 +94,6 @@ bool FStreamingLevelModel::IsGoodToDrop(const TSharedPtr<FLevelDragDropOp>& Op)
 	return false;
 }
 
-void FStreamingLevelModel::GetGridItemTooltipFields(TArray< TPair<TAttribute<FText>, TAttribute<FText>> >& CustomFields) const
-{
-
-}
-
 //bool FStreamingLevelModel::IsReadOnly() const
 //{
 //	ULevel* Level = GetLevelObject();
diff --git a/Engine/Source/Editor/WorldBrowser/Private/StreamingLevels/StreamingLevelModel.h b/Engine/Source/Editor/WorldBrowser/Private/StreamingLevels/StreamingLevelModel.h
--- a/Engine/Source/Editor/WorldBrowser/Private/StreamingLevels/StreamingLevelModel.h
+++ b/Engine/Source/Editor/WorldBrowser/Private/StreamingLevels/StreamingLevelModel.h
@@ -35,7 +35,6 @@ public:
 	virtual void Update() OVERRIDE;
 	virtual void OnDrop(const TSharedPtr<FLevelDragDropOp>& Op) OVERRIDE;
 	virtual bool IsGoodToDrop(const TSharedPtr<FLevelDragDropOp>& Op) const OVERRIDE;
-	virtual void GetGridItemTooltipFields(TArray< TPair<TAttribute<FText>, TAttribute<FText>> >& CustomFields) const OVERRIDE;
 	virtual UClass* GetStreamingClass() const OVERRIDE;
 	// FLevelModel interface end
 		
diff --git a/Engine/Source/Editor/WorldBrowser/Private/Tiles/SWorldComposition.cpp b/Engine/Source/Editor/WorldBrowser/Private/Tiles/SWorldComposition.cpp
--- a/Engine/Source/Editor/WorldBrowser/Private/Tiles/SWorldComposition.cpp
+++ b/Engine/Source/Editor/WorldBrowser/Private/Tiles/SWorldComposition.cpp
@@ -140,7 +140,7 @@ public:
 	}
 	
 	/**  Add specified item to the grid view */
-	void AddItem(const TSharedPtr<FLevelModel>& LevelModel)
+	void AddItem(TSharedPtr<FWorldTileModel> LevelModel)
 	{
 		auto NewNode = SNew(SWorldTileItem)
 							.InWorldModel(WorldModel)
@@ -151,7 +151,7 @@ public:
 	}
 	
 	/**  Remove specified item from the grid view */
-	void RemoveItem(const TSharedPtr<FLevelModel>& LevelModel)
+	void RemoveItem(TSharedPtr<FLevelModel> LevelModel)
 	{
 		TSharedRef<SNode>* pItem = NodeToWidgetLookup.Find(LevelModel->GetNodeObject());
 		if (pItem == NULL)
@@ -172,7 +172,7 @@ public:
 		FLevelModelList AllLevels = WorldModel->GetAllLevels();
 		for (auto It = AllLevels.CreateConstIterator(); It; ++It)
 		{
-			AddItem(*It);
+			AddItem(StaticCastSharedPtr<FWorldTileModel>(*It));
 		}
 	}
 		
diff --git a/Engine/Source/Editor/WorldBrowser/Private/Tiles/SWorldTileItem.cpp b/Engine/Source/Editor/WorldBrowser/Private/Tiles/SWorldTileItem.cpp
--- a/Engine/Source/Editor/WorldBrowser/Private/Tiles/SWorldTileItem.cpp
+++ b/Engine/Source/Editor/WorldBrowser/Private/Tiles/SWorldTileItem.cpp
@@ -4,6 +4,7 @@
 #include "ObjectTools.h"
 #include "EdGraphUtilities.h"
 #include "SWorldTileItem.h"
+#include "WorldTileDetails.h"
 
 #define LOCTEXT_NAMESPACE "WorldBrowser"
 
@@ -127,13 +128,13 @@ void FTileItemThumbnail::UpdateThumbnail()
 //----------------------------------------------------------------
 SWorldTileItem::~SWorldTileItem()
 {
-	LevelModel->ChangedEvent.RemoveAll(this);
+	TileModel->ChangedEvent.RemoveAll(this);
 }
 
 void SWorldTileItem::Construct(const FArguments& InArgs)
 {
 	WorldModel = InArgs._InWorldModel;
-	LevelModel = InArgs._InItemModel;
+	TileModel = InArgs._InItemModel;
 
 	ProgressBarImage = FEditorStyle::GetBrush(TEXT("ProgressBar.Marquee"));
 	
@@ -141,7 +142,7 @@ void SWorldTileItem::Construct(const FArguments& InArgs)
 	bIsDragging = false;
 	bAffectedByMarquee = false;
 
-	Thumbnail = MakeShareable(new FTileItemThumbnail(InArgs._ThumbnailRenderTarget, LevelModel));
+	Thumbnail = MakeShareable(new FTileItemThumbnail(InArgs._ThumbnailRenderTarget, TileModel));
 	ThumbnailViewport = SNew(SViewport)
 							.EnableGammaCorrection(false);
 							
@@ -149,7 +150,7 @@ void SWorldTileItem::Construct(const FArguments& InArgs)
 	// This will grey out tile in case level is not loaded or locked
 	ThumbnailViewport->SetEnabled(TAttribute<bool>(this, &SWorldTileItem::IsItemEnabled));
 
-	LevelModel->ChangedEvent.AddSP(this, &SWorldTileItem::RequestRefresh);
+	TileModel->ChangedEvent.AddSP(this, &SWorldTileItem::RequestRefresh);
 			
 	ChildSlot
 	[
@@ -169,14 +170,13 @@ void SWorldTileItem::RequestRefresh()
 
 UObject* SWorldTileItem::GetObjectBeingDisplayed() const
 {
-	return LevelModel->GetNodeObject();
+	return TileModel->GetNodeObject();
 }
 
 BEGIN_SLATE_FUNCTION_BUILD_OPTIMIZATION
 TSharedRef<SToolTip> SWorldTileItem::CreateToolTipWidget()
 {
 	TSharedPtr<SToolTip>		TooltipWidget;
-	TSharedPtr<SVerticalBox>	TooltipSectionsWidget;
 	
 	SAssignNew(TooltipWidget, SToolTip)
 	.TextMargin(1)
@@ -186,7 +186,7 @@ TSharedRef<SToolTip> SWorldTileItem::CreateToolTipWidget()
 		.Padding(6)
 		.BorderImage(FEditorStyle::GetBrush("ContentBrowser.TileViewTooltip.NonContentBorder"))
 		[
-			SAssignNew(TooltipSectionsWidget, SVerticalBox)
+			SNew(SVerticalBox)
 
 			// Level name section
 			+SVerticalBox::Slot()
@@ -218,48 +218,81 @@ TSharedRef<SToolTip> SWorldTileItem::CreateToolTipWidget()
 					]
 				]
 			]
-		]
-	];
-
-	TArray< TPair<TAttribute<FText>, TAttribute<FText>> > CustomFields;
-	LevelModel->GetGridItemTooltipFields(CustomFields);
-
-	if (CustomFields.Num() > 0)
-	{
-		TSharedPtr<SUniformGridPanel> TooltipFieldsWidget;
-		// Add section for custom fields
-		TooltipSectionsWidget->AddSlot()
+			
+			// Tile info section
+			+SVerticalBox::Slot()
 			.AutoHeight()
 			.Padding(0,0,0,4)
-		[
+			[
 				SNew(SBorder)
 				.Padding(6)
 				.BorderImage(FEditorStyle::GetBrush("ContentBrowser.TileViewTooltip.ContentBorder"))
 				[
-					SAssignNew(TooltipFieldsWidget, SUniformGridPanel)
-				]
-		];
+					SNew(SUniformGridPanel)
 
-		// Add custom fields
-		for (int32 FieldIdx = 0; FieldIdx < CustomFields.Num(); ++FieldIdx)
-		{
-			// Name:
-			TooltipFieldsWidget->AddSlot(0, FieldIdx)
-				.HAlign(HAlign_Left)
-			[
-				SNew(STextBlock)
-				.Text(CustomFields[FieldIdx].Key) 
-			];
-		
-			// Value:
-			TooltipFieldsWidget->AddSlot(1, FieldIdx)
-				.HAlign(HAlign_Right)
-			[
-				SNew(STextBlock)
-				.Text(CustomFields[FieldIdx].Value) 
-			];
-		}
-	}
+					// Tile position
+					+SUniformGridPanel::Slot(0, 0)
+					.HAlign(HAlign_Left)
+					[
+						SNew(STextBlock)
+						.Text(LOCTEXT("Item_OriginOffset", "Position:")) 
+					]
+					
+					+SUniformGridPanel::Slot(1, 0)
+					.HAlign(HAlign_Right)
+					[
+						SNew(STextBlock)
+						.Text(this, &SWorldTileItem::GetPositionText) 
+					]
+
+					// Tile bounds
+					+SUniformGridPanel::Slot(0, 1)
+					.HAlign(HAlign_Left)
+					[
+						SNew(STextBlock)
+						.Text(LOCTEXT("Item_BoundsExtent", "Extent:")) 
+					]
+					
+					+SUniformGridPanel::Slot(1, 1)
+					.HAlign(HAlign_Right)
+					[
+						SNew(STextBlock)
+						.Text(this, &SWorldTileItem::GetBoundsExtentText) 
+					]
+
+					// Layer name
+					+SUniformGridPanel::Slot(0, 2)
+					.HAlign(HAlign_Left)
+					[
+						SNew(STextBlock)
+						.Text(LOCTEXT("Item_Name", "Layer Name:")) 
+					]
+					
+					+SUniformGridPanel::Slot(1, 2)
+					.HAlign(HAlign_Right)
+					[
+						SNew(STextBlock)
+						.Text(this, &SWorldTileItem::GetLevelLayerNameText) 
+					]
+
+					// Streaming distance
+					+SUniformGridPanel::Slot(0, 3)
+					.HAlign(HAlign_Left)
+					[
+						SNew(STextBlock)
+						.Text(LOCTEXT("Item_Distance", "Streaming Distance:")) 
+					]
+					
+					+SUniformGridPanel::Slot(1, 3)
+					.HAlign(HAlign_Right)
+					[
+						SNew(STextBlock)
+						.Text(this, &SWorldTileItem::GetLevelLayerDistanceText) 
+					]
+				]
+			]
+		]
+	];
 
 	return TooltipWidget.ToSharedRef();
 }
@@ -267,12 +300,12 @@ END_SLATE_FUNCTION_BUILD_OPTIMIZATION
 
 FVector2D SWorldTileItem::GetPosition() const
 {
-	return LevelModel->GetLevelPosition2D();
+	return TileModel->GetLevelPosition2D();
 }
 
-const TSharedPtr<FLevelModel>& SWorldTileItem::GetLevelModel() const
+TSharedPtr<FLevelModel> SWorldTileItem::GetLevelModel() const
 {
-	return LevelModel;
+	return StaticCastSharedPtr<FLevelModel>(TileModel);
 }
 
 const FSlateBrush* SWorldTileItem::GetShadowBrush(bool bSelected) const
@@ -282,17 +315,17 @@ const FSlateBrush* SWorldTileItem::GetShadowBrush(bool bSelected) const
 
 FOptionalSize SWorldTileItem::GetItemWidth() const
 {
-	return FOptionalSize(LevelModel->GetLevelSize2D().X);
+	return FOptionalSize(TileModel->GetLevelSize2D().X);
 }
 
 FOptionalSize SWorldTileItem::GetItemHeight() const
 {
-	return FOptionalSize(LevelModel->GetLevelSize2D().Y);
+	return FOptionalSize(TileModel->GetLevelSize2D().Y);
 }
 
 FSlateRect SWorldTileItem::GetItemRect() const
 {
-	FVector2D LevelSize = LevelModel->GetLevelSize2D();
+	FVector2D LevelSize = TileModel->GetLevelSize2D();
 	FVector2D LevelPos = GetPosition();
 	return FSlateRect(LevelPos, LevelPos + LevelSize);
 }
@@ -300,7 +333,7 @@ FSlateRect SWorldTileItem::GetItemRect() const
 TSharedPtr<IToolTip> SWorldTileItem::GetToolTip()
 {
 	// Hide tooltip in case item is being dragged now
-	if (LevelModel->GetLevelTranslationDelta().Size() > KINDA_SMALL_NUMBER)
+	if (TileModel->GetLevelTranslationDelta().Size() > KINDA_SMALL_NUMBER)
 	{
 		return NULL;
 	}
@@ -311,7 +344,7 @@ TSharedPtr<IToolTip> SWorldTileItem::GetToolTip()
 FVector2D SWorldTileItem::GetDesiredSizeForMarquee() const
 {
 	// we don't want to select items in non visible layers
-	if (!WorldModel->PassesAllFilters(LevelModel))
+	if (!WorldModel->PassesAllFilters(TileModel))
 	{
 		return FVector2D::ZeroVector;
 	}
@@ -321,7 +354,7 @@ FVector2D SWorldTileItem::GetDesiredSizeForMarquee() const
 
 FVector2D SWorldTileItem::ComputeDesiredSize() const
 {
-	return LevelModel->GetLevelSize2D();
+	return TileModel->GetLevelSize2D();
 }
 
 int32 SWorldTileItem::OnPaint(const FGeometry& AllottedGeometry, const FSlateRect& ClippingRect, 
@@ -377,7 +410,7 @@ int32 SWorldTileItem::OnPaint(const FGeometry& AllottedGeometry, const FSlateRec
 		//}
 
 		// Draw progress bar if level is currently loading
-		if (LevelModel->IsLoading())
+		if (TileModel->IsLoading())
 		{
 			const float ProgressBarAnimOffset = ProgressBarImage->ImageSize.X * CurveSequence.GetLerpLooping() / AllottedGeometry.Scale;
 			const float ProgressBarImageSize = ProgressBarImage->ImageSize.X / AllottedGeometry.Scale;
@@ -409,7 +442,7 @@ FReply SWorldTileItem::OnMouseButtonDoubleClick(const FGeometry& InMyGeometry, c
 {
 	if (OnHitTest(InMyGeometry, InMouseEvent.GetScreenSpacePosition()))
 	{
-		LevelModel->MakeLevelCurrent();
+		TileModel->MakeLevelCurrent();
 		return FReply::Handled();
 	}
 	
@@ -419,35 +452,57 @@ FReply SWorldTileItem::OnMouseButtonDoubleClick(const FGeometry& InMyGeometry, c
 bool SWorldTileItem::OnHitTest(const FGeometry& MyGeometry, FVector2D InAbsoluteCursorPosition)
 {
 	FVector2D CursorOffset = (InAbsoluteCursorPosition - MyGeometry.AbsolutePosition)/MyGeometry.Scale;
-	FVector2D LevelPos = LevelModel->GetLevelPosition2D();
+	FVector2D LevelPos = TileModel->GetLevelPosition2D();
 	FVector2D CursorWorldPos = LevelPos + CursorOffset;
-	return LevelModel->HitTest2D(CursorWorldPos);
+	return TileModel->HitTest2D(CursorWorldPos);
 }
 
 FString SWorldTileItem::GetLevelNameText() const
 {
-	return LevelModel->GetDisplayName();
+	return TileModel->GetDisplayName();
+}
+
+FText SWorldTileItem::GetPositionText() const
+{
+	FIntPoint Position = TileModel->GetRelativeLevelPosition();
+	return FText::FromString(FString::Printf(TEXT("%d, %d"), Position.X, Position.Y));
+}
+
+FText SWorldTileItem::GetBoundsExtentText() const
+{
+	FVector2D Size = TileModel->GetLevelSize2D();
+	return FText::FromString(FString::Printf(TEXT("%d, %d"), FMath::RoundToInt(Size.X*0.5f), FMath::RoundToInt(Size.Y*0.5f)));
+}
+
+FText SWorldTileItem::GetLevelLayerNameText() const
+{
+	return FText::FromString(TileModel->TileDetails->Layer.Name);
+}
+
+FText SWorldTileItem::GetLevelLayerDistanceText() const
+{
+	return FText::AsNumber(TileModel->TileDetails->Layer.StreamingDistance);
 }
 
 bool SWorldTileItem::IsItemEditable() const
 {
-	return LevelModel->IsEditable();
+	return TileModel->IsEditable();
 }
 
 bool SWorldTileItem::IsItemSelected() const
 {
-	return LevelModel->GetLevelSelectionFlag();
+	return TileModel->GetLevelSelectionFlag();
 }
 
 bool SWorldTileItem::IsItemEnabled() const
 {
 	if (WorldModel->IsSimulating())
 	{
-		return LevelModel->IsVisible();
+		return TileModel->IsVisible();
 	}
 	else
 	{
-		return LevelModel->IsEditable();
+		return TileModel->IsEditable();
 	}
 }
 
diff --git a/Engine/Source/Editor/WorldBrowser/Private/Tiles/SWorldTileItem.h b/Engine/Source/Editor/WorldBrowser/Private/Tiles/SWorldTileItem.h
--- a/Engine/Source/Editor/WorldBrowser/Private/Tiles/SWorldTileItem.h
+++ b/Engine/Source/Editor/WorldBrowser/Private/Tiles/SWorldTileItem.h
@@ -42,9 +42,9 @@ public:
 	SLATE_BEGIN_ARGS(SWorldTileItem)
 	{}
 		/** The world data */
-		SLATE_ARGUMENT(TSharedPtr<FLevelCollectionModel>, InWorldModel)
+		SLATE_ARGUMENT(TSharedPtr<FWorldTileCollectionModel>, InWorldModel)
 		/** Data for the asset this item represents */
-		SLATE_ARGUMENT(TSharedPtr<FLevelModel>, InItemModel)
+		SLATE_ARGUMENT(TSharedPtr<FWorldTileModel>, InItemModel)
 		//
 		SLATE_ARGUMENT(FSlateTextureRenderTarget2DResource*, ThumbnailRenderTarget)
 	SLATE_END_ARGS()
@@ -64,7 +64,7 @@ public:
 	void RequestRefresh();
 
 	/** @return LevelModel associated with this item */
-	const TSharedPtr<FLevelModel>&	GetLevelModel() const;
+	TSharedPtr<FLevelModel>	GetLevelModel() const;
 	
 	/** @return Item width in world units */
 	FOptionalSize GetItemWidth() const;
@@ -97,18 +97,22 @@ private:
 
 	TSharedRef<SToolTip> CreateToolTipWidget();
 	
-	//
+	/** Tile tooltips fields */
 	FString GetLevelNameText() const;
-
+	FText GetPositionText() const;
+	FText GetBoundsExtentText() const;
+	FText GetLevelLayerNameText() const;
+	FText GetLevelLayerDistanceText() const;
+	
 public:
 	bool							bAffectedByMarquee;
 
 private:
 	/** The world data */
-	TSharedPtr<FLevelCollectionModel> WorldModel;
+	TSharedPtr<FWorldTileCollectionModel> WorldModel;
 
 	/** The data for this item */
-	TSharedPtr<FLevelModel>			LevelModel;
+	TSharedPtr<FWorldTileModel>		TileModel;
 
 	mutable bool					bNeedRefresh;
 	bool							bIsDragging;
diff --git a/Engine/Source/Editor/WorldBrowser/Private/Tiles/WorldTileModel.cpp b/Engine/Source/Editor/WorldBrowser/Private/Tiles/WorldTileModel.cpp
--- a/Engine/Source/Editor/WorldBrowser/Private/Tiles/WorldTileModel.cpp
+++ b/Engine/Source/Editor/WorldBrowser/Private/Tiles/WorldTileModel.cpp
@@ -193,35 +193,6 @@ bool FWorldTileModel::IsGoodToDrop(const TSharedPtr<FLevelDragDropOp>& Op) const
 	return true;
 }
 
-void FWorldTileModel::GetGridItemTooltipFields(TArray< TPair<TAttribute<FText>, TAttribute<FText>> >& CustomFields) const
-{
-	typedef TPair<TAttribute<FText>, TAttribute<FText>>	FTooltipField;
-
-	// Position
-	FTooltipField PositionField;
-	PositionField.Key				= LOCTEXT("Item_OriginOffset", "Position:");
-	PositionField.Value				= TAttribute<FText>(this, &FWorldTileModel::GetPositionText);
-	CustomFields.Add(PositionField);
-
-	// Bounds extent
-	FTooltipField BoundsExtentField;
-	BoundsExtentField.Key			= LOCTEXT("Item_BoundsExtent", "Extent:");
-	BoundsExtentField.Value			= TAttribute<FText>(this, &FWorldTileModel::GetBoundsExtentText);
-	CustomFields.Add(BoundsExtentField);
-
-	// Layer name
-	FTooltipField LayerNameField;
-	LayerNameField.Key				= LOCTEXT("Item_Name", "Layer Name:");
-	LayerNameField.Value			= TAttribute<FText>(this, &FWorldTileModel::GetLevelLayerNameText);
-	CustomFields.Add(LayerNameField);
-	
-	// Streaming distance
-	FTooltipField StreamingDistanceField;
-	StreamingDistanceField.Key		= LOCTEXT("Item_Distance", "Streaming Distance:");
-	StreamingDistanceField.Value	= TAttribute<FText>(this, &FWorldTileModel::GetLevelLayerDistanceText);
-	CustomFields.Add(StreamingDistanceField);
-}
-
 bool FWorldTileModel::ShouldBeVisible(FBox EditableArea) const
 {
 	if (IsRootTile())
@@ -808,28 +779,6 @@ void FWorldTileModel::OnZOrderPropertyChanged()
 	OnLevelInfoUpdated();
 }
 
-FText FWorldTileModel::GetPositionText() const
-{
-	FIntPoint Position = GetRelativeLevelPosition();
-	return FText::FromString(FString::Printf(TEXT("%d, %d"), Position.X, Position.Y));
-}
-
-FText FWorldTileModel::GetBoundsExtentText() const
-{
-	FVector2D Size = GetLevelSize2D();
-	return FText::FromString(FString::Printf(TEXT("%d, %d"), FMath::RoundToInt(Size.X*0.5f), FMath::RoundToInt(Size.Y*0.5f)));
-}
-
-FText FWorldTileModel::GetLevelLayerNameText() const
-{
-	return FText::FromString(TileDetails->Layer.Name);
-}
-
-FText FWorldTileModel::GetLevelLayerDistanceText() const
-{
-	return FText::AsNumber(TileDetails->Layer.StreamingDistance);
-}
-
 bool FWorldTileModel::CreateAdjacentLandscapeProxy(ALandscapeProxy* SourceLandscape, FIntPoint SourceTileOffset, FWorldTileModel::EWorldDirections InWhere)
 {
 	if (!IsLoaded())	
diff --git a/Engine/Source/Editor/WorldBrowser/Private/Tiles/WorldTileModel.h b/Engine/Source/Editor/WorldBrowser/Private/Tiles/WorldTileModel.h
--- a/Engine/Source/Editor/WorldBrowser/Private/Tiles/WorldTileModel.h
+++ b/Engine/Source/Editor/WorldBrowser/Private/Tiles/WorldTileModel.h
@@ -75,7 +75,6 @@ public:
 	virtual FVector2D GetLevelSize2D() const OVERRIDE;
 	virtual void OnDrop(const TSharedPtr<FLevelDragDropOp>& Op) OVERRIDE;
 	virtual bool IsGoodToDrop(const TSharedPtr<FLevelDragDropOp>& Op) const OVERRIDE;
-	virtual void GetGridItemTooltipFields(TArray< TPair<TAttribute<FText>, TAttribute<FText>> >& CustomFields) const OVERRIDE;
 	virtual void OnLevelAddedToWorld(ULevel* InLevel) OVERRIDE;
 	virtual void OnLevelRemovedFromWorld() OVERRIDE;
 	virtual void OnParentChanged() OVERRIDE;
@@ -181,12 +180,6 @@ private:
 	/** Handler for ZOrder chnages event from Tile details object  */
 	void OnZOrderPropertyChanged();
 	
-	/** Tile tooltips support */
-	FText GetPositionText() const;
-	FText GetBoundsExtentText() const;
-	FText GetLevelLayerNameText() const;
-	FText GetLevelLayerDistanceText() const;
-	
 public:
 	/** This tile index in world composition tile list */
 	int32									TileIdx;
