UE4 fix non-unity errors

diff --git a/Engine/Source/Editor/WorldBrowser/Private/Tiles/SWorldComposition.cpp b/Engine/Source/Editor/WorldBrowser/Private/Tiles/SWorldComposition.cpp
--- a/Engine/Source/Editor/WorldBrowser/Private/Tiles/SWorldComposition.cpp
+++ b/Engine/Source/Editor/WorldBrowser/Private/Tiles/SWorldComposition.cpp
@@ -4,6 +4,7 @@
 #include "GraphEditor.h"
 #include "EdGraphUtilities.h"
 
+#include "WorldTileModel.h"
 #include "SWorldComposition.h"
 #include "SWorldTileItem.h"
 #include "SWorldLayers.h"
diff --git a/Engine/Source/Editor/WorldBrowser/Private/Tiles/SWorldTileItem.cpp b/Engine/Source/Editor/WorldBrowser/Private/Tiles/SWorldTileItem.cpp
--- a/Engine/Source/Editor/WorldBrowser/Private/Tiles/SWorldTileItem.cpp
+++ b/Engine/Source/Editor/WorldBrowser/Private/Tiles/SWorldTileItem.cpp
@@ -3,6 +3,7 @@
 
 #include "ObjectTools.h"
 #include "EdGraphUtilities.h"
+#include "WorldTileCollectionModel.h"
 #include "SWorldTileItem.h"
 #include "WorldTileDetails.h"
 
