Fixed world bounds not showing in the world browser status bar

diff --git a/Engine/Source/Editor/WorldBrowser/Private/Tiles/SWorldComposition.cpp b/Engine/Source/Editor/WorldBrowser/Private/Tiles/SWorldComposition.cpp
--- a/Engine/Source/Editor/WorldBrowser/Private/Tiles/SWorldComposition.cpp
+++ b/Engine/Source/Editor/WorldBrowser/Private/Tiles/SWorldComposition.cpp
@@ -1199,7 +1199,9 @@ FString SWorldComposition::GetMouseLocationText() const
 FString	SWorldComposition::GetMarqueeSelectionSizeText() const
 {
 	FVector2D MarqueeSize = GridView->GetMarqueeWorldSize();
-	if (MarqueeSize.Size() > 0)
+	
+	if (MarqueeSize.X > 0 && 
+		MarqueeSize.Y > 0)
 	{
 		return FString::Printf(TEXT("%d x %d"), FMath::RoundToInt(MarqueeSize.X), FMath::RoundToInt(MarqueeSize.Y));
 	}
@@ -1212,7 +1214,9 @@ FString	SWorldComposition::GetMarqueeSelectionSizeText() const
 FString SWorldComposition::GetWorldSizeText() const
 {
 	FIntPoint WorldSize = TileWorldModel->GetWorldSize();
-	if (WorldSize.Size() > 0)
+	
+	if (WorldSize.X > 0 && 
+		WorldSize.Y > 0)
 	{
 		return FString::Printf(TEXT("%d x %d"), WorldSize.X, WorldSize.Y);
 	}
