Fixed importing landscape tile in World Browser when world origin rebasing is disabled

diff --git a/Engine/Source/Editor/WorldBrowser/Private/Tiles/WorldTileCollectionModel.cpp b/Engine/Source/Editor/WorldBrowser/Private/Tiles/WorldTileCollectionModel.cpp
--- a/Engine/Source/Editor/WorldBrowser/Private/Tiles/WorldTileCollectionModel.cpp
+++ b/Engine/Source/Editor/WorldBrowser/Private/Tiles/WorldTileCollectionModel.cpp
@@ -1231,7 +1231,7 @@ void FWorldTileCollectionModel::AddLandscapeProxy_Executed(FWorldTileModel::EWor
 		ALandscapeProxy* SourceLandscape = LandscapeTileModel->GetLandcape();
 		FIntPoint SourceTileOffset = LandscapeTileModel->GetAbsoluteLevelPosition();
 
-		NewLevelModel->Shelve();
+		NewLevelModel->SetVisible(false);
 		NewLevelModel->CreateAdjacentLandscapeProxy(SourceLandscape, SourceTileOffset, InWhere);
 		ShowLevels(Levels);
 	}
@@ -1376,7 +1376,7 @@ void FWorldTileCollectionModel::ImportTiledLandscape_Executed()
 					PopulateLevelsList();
 					TSharedPtr<FWorldTileModel> NewTileModel = StaticCastSharedPtr<FWorldTileModel>(FindLevelModel(NewWorld->GetOutermost()->GetFName()));
 					// Hide level, so we do not depend on a current world origin
-					NewTileModel->Shelve();
+					NewTileModel->SetVisible(false);
 					
 					// Create landscape (actor/proxy) in a new level
 					ALandscapeProxy* NewLandscape = NewTileModel->ImportLandscape(TileImportSettings);
