Tiled landscape import improvements
Added option for weight-blended layers
Added ability to import only part of weightmaps

diff --git a/Engine/Source/Editor/WorldBrowser/Private/LevelCollectionModel.h b/Engine/Source/Editor/WorldBrowser/Private/LevelCollectionModel.h
--- a/Engine/Source/Editor/WorldBrowser/Private/LevelCollectionModel.h
+++ b/Engine/Source/Editor/WorldBrowser/Private/LevelCollectionModel.h
@@ -557,9 +557,18 @@ struct FTiledLandscapeImportSettings
 
 
 	TWeakObjectPtr<UMaterialInterface>	LandscapeMaterial;
-	
-	TArray<FName>						LandscapeLayerNameList;
-	// list of weightmap files per each layer
-	TArray<TArray<FString>>				WeightmapFileList;
 
+	// Landscape layers 
+	struct LandscapeLayerSettings
+	{
+		LandscapeLayerSettings() 
+			: bNoBlendWeight(false)
+		{}
+
+		FName						Name;
+		bool						bNoBlendWeight;
+		TMap<FIntPoint, FString>	WeightmapFiles;
+	};
+	
+	TArray<LandscapeLayerSettings>		LandscapeLayerSettingsList;
 };
\ No newline at end of file
diff --git a/Engine/Source/Editor/WorldBrowser/Private/Tiles/STiledLandscapeImportDlg.cpp b/Engine/Source/Editor/WorldBrowser/Private/Tiles/STiledLandscapeImportDlg.cpp
--- a/Engine/Source/Editor/WorldBrowser/Private/Tiles/STiledLandscapeImportDlg.cpp
+++ b/Engine/Source/Editor/WorldBrowser/Private/Tiles/STiledLandscapeImportDlg.cpp
@@ -15,9 +15,9 @@ static int32 CalcLandscapeSquareResolution(int32 ComponetsNumX, int32 SectionNum
 }
 
 /**
- *	Returns heightmap tile coordinates extracted from a specified tile filename
+ *	Returns tile coordinates extracted from a specified tile filename
  */
-static FIntPoint ExtractHeighmapTileCoordinates(FString BaseFilename)
+static FIntPoint ExtractTileCoordinates(FString BaseFilename)
 {
 	//We expect file name in form: <tilename>_x<number>_y<number>
 	FIntPoint ResultPosition(-1,-1);
@@ -187,7 +187,7 @@ void STiledLandcapeImportDlg::Construct(const FArguments& InArgs, TSharedPtr<SWi
 			.AutoHeight()
 			.Padding(FEditorStyle::GetMargin("StandardDialog.ContentPadding"))
 			[
-				SAssignNew(LayerDataListView, SListView<TSharedPtr<FLandscapeImportLayerData>>)
+				SAssignNew(LayerDataListView, SListView<TSharedPtr<FTiledLandscapeImportSettings::LandscapeLayerSettings>>)
 				.ListItemsSource( &LayerDataList )
 				.OnGenerateRow( this, &STiledLandcapeImportDlg::OnGenerateWidgetForLayerDataListView )
 				.SelectionMode(ESelectionMode::None)
@@ -297,22 +297,39 @@ FText STiledLandcapeImportDlg::GetTileConfigurationText() const
 	return GenerateConfigurationText(ImportSettings.ComponentsNum, ImportSettings.SectionsPerComponent, ImportSettings.QuadsPerSection);
 }
 
-TSharedRef<ITableRow> STiledLandcapeImportDlg::OnGenerateWidgetForLayerDataListView(TSharedPtr<FLandscapeImportLayerData> InLayerData, const TSharedRef<STableViewBase>& OwnerTable)
+TSharedRef<ITableRow> STiledLandcapeImportDlg::OnGenerateWidgetForLayerDataListView(
+	TSharedPtr<FTiledLandscapeImportSettings::LandscapeLayerSettings> InLayerData, 
+	const TSharedRef<STableViewBase>& OwnerTable)
 {
-	return 	SNew(STableRow< TSharedRef<FLandscapeImportLayerData> >, OwnerTable)
+	return 	SNew(STableRow<TSharedRef<FTiledLandscapeImportSettings::LandscapeLayerSettings>>, OwnerTable)
 			[
 				SNew(SBorder)
 				[
 					SNew(SHorizontalBox)
-				
+
+					// Layer name
 					+SHorizontalBox::Slot()
 					.VAlign(VAlign_Center)
 					.HAlign(HAlign_Left)
 					.FillWidth(1)
 					[
-						SNew(STextBlock).Text(InLayerData->LayerName)
+						SNew(STextBlock).Text(InLayerData->Name.ToString())
 					]
 
+					// Blend option
+					+SHorizontalBox::Slot()
+					.VAlign(VAlign_Center)
+					.HAlign(HAlign_Center)
+					.Padding(2)
+					.AutoWidth()
+					[
+						SNew(SCheckBox)
+						.IsChecked(this, &STiledLandcapeImportDlg::GetLayerBlendState, InLayerData)
+						.OnCheckStateChanged(this, &STiledLandcapeImportDlg::OnLayerBlendStateChanged, InLayerData)
+						.ToolTipText(LOCTEXT("TiledLandscapeImport_BlendOption", "Weight-Blended Layer"))
+					]
+					
+					// Number of selected weightmaps
 					+SHorizontalBox::Slot()
 					.VAlign(VAlign_Center)
 					.HAlign(HAlign_Right)
@@ -322,6 +339,7 @@ TSharedRef<ITableRow> STiledLandcapeImportDlg::OnGenerateWidgetForLayerDataListV
 						SNew(STextBlock).Text(this, &STiledLandcapeImportDlg::GetWeightmapCountText, InLayerData)
 					]
 					
+					// Button for selecting weightmap files
 					+SHorizontalBox::Slot()
 					.VAlign(VAlign_Center)
 					.HAlign(HAlign_Right)
@@ -447,7 +465,7 @@ FReply STiledLandcapeImportDlg::OnClickedSelectHeightmapTiles()
 						break;
 					}
 
-					FIntPoint TileCoordinate = ExtractHeighmapTileCoordinates(FPaths::GetBaseFilename(Filename));
+					FIntPoint TileCoordinate = ExtractTileCoordinates(FPaths::GetBaseFilename(Filename));
 					if (TileCoordinate.GetMin() < 0)
 					{
 						bValidTiles = false;
@@ -471,9 +489,9 @@ FReply STiledLandcapeImportDlg::OnClickedSelectHeightmapTiles()
 	return FReply::Handled();
 }
 
-FReply STiledLandcapeImportDlg::OnClickedSelectWeighmapTiles(TSharedPtr<FLandscapeImportLayerData> InLayerData)
+FReply STiledLandcapeImportDlg::OnClickedSelectWeighmapTiles(TSharedPtr<FTiledLandscapeImportSettings::LandscapeLayerSettings> InLayerData)
 {
-	InLayerData->WeighmapFileList.Empty();
+	InLayerData->WeightmapFiles.Empty();
 
 	IDesktopPlatform* DesktopPlatform = FDesktopPlatformModule::Get();
 	if (DesktopPlatform)
@@ -482,6 +500,8 @@ FReply STiledLandcapeImportDlg::OnClickedSelectWeighmapTiles(TSharedPtr<FLandsca
 		{
 			void* ParentWindowWindowHandle = ParentWindow->GetNativeWindow()->GetOSWindowHandle();
 
+			TArray<FString> WeightmapFilesList;
+
 			const FString ImportFileTypes = TEXT("Raw weigtmap tiles (*.raw)|*.raw");
 			bool bOpened = DesktopPlatform->OpenFileDialog(
 								ParentWindowWindowHandle,
@@ -490,10 +510,16 @@ FReply STiledLandcapeImportDlg::OnClickedSelectWeighmapTiles(TSharedPtr<FLandsca
 								TEXT(""),
 								*ImportFileTypes,
 								EFileDialogFlags::Multiple,
-								InLayerData->WeighmapFileList);
+								WeightmapFilesList);
 
-			if (bOpened	&& InLayerData->WeighmapFileList.Num())
+			if (bOpened	&& WeightmapFilesList.Num())
 			{
+				for (FString WeighmapFile : WeightmapFilesList)
+				{
+					FIntPoint TileCoordinate = ExtractTileCoordinates(FPaths::GetBaseFilename(WeighmapFile));
+					InLayerData->WeightmapFiles.Add(TileCoordinate, WeighmapFile);
+				}
+							
 				// TODO: check if it's a valid weightmaps
 			}
 		}
@@ -510,19 +536,13 @@ bool STiledLandcapeImportDlg::IsImportEnabled() const
 FReply STiledLandcapeImportDlg::OnClickedImport()
 {
 	// copy weightmaps list data to an import structure  
-	ImportSettings.WeightmapFileList.Empty();
-	ImportSettings.WeightmapFileList.SetNum(ImportSettings.LandscapeLayerNameList.Num());
-	check(ImportSettings.WeightmapFileList.Num() == LayerDataList.Num());
+	ImportSettings.LandscapeLayerSettingsList.Empty();
+
 	for (int32 LayerIdx = 0; LayerIdx < LayerDataList.Num(); ++LayerIdx)
 	{
-		ImportSettings.WeightmapFileList[LayerIdx] = LayerDataList[LayerIdx]->WeighmapFileList;
-		
-		// we assume weightmaps tiles will match heightmap tiles 
-		ImportSettings.WeightmapFileList[LayerIdx].Sort();
+		ImportSettings.LandscapeLayerSettingsList.Add(*(LayerDataList[LayerIdx].Get()));
 	}
-
-	ImportSettings.HeightmapFileList.Sort();
-			
+					
 	ParentWindow->RequestDestroyWindow();
 	return FReply::Handled();	
 }
@@ -623,12 +643,22 @@ FText STiledLandcapeImportDlg::GetImportSummaryText() const
 		);
 }
 
-FText STiledLandcapeImportDlg::GetWeightmapCountText(TSharedPtr<FLandscapeImportLayerData> InLayerData) const
+FText STiledLandcapeImportDlg::GetWeightmapCountText(TSharedPtr<FTiledLandscapeImportSettings::LandscapeLayerSettings> InLayerData) const
 {
-	int32 NumWeighmaps = InLayerData.IsValid() ? InLayerData->WeighmapFileList.Num() : 0;
+	int32 NumWeighmaps = InLayerData.IsValid() ? InLayerData->WeightmapFiles.Num() : 0;
 	return FText::AsNumber(NumWeighmaps);
 }
 
+ESlateCheckBoxState::Type STiledLandcapeImportDlg::GetLayerBlendState(TSharedPtr<FTiledLandscapeImportSettings::LandscapeLayerSettings> InLayerData) const
+{
+	return (InLayerData->bNoBlendWeight ? ESlateCheckBoxState::Unchecked : ESlateCheckBoxState::Checked);
+}
+
+void STiledLandcapeImportDlg::OnLayerBlendStateChanged(ESlateCheckBoxState::Type NewState, TSharedPtr<FTiledLandscapeImportSettings::LandscapeLayerSettings> InLayerData)
+{
+	InLayerData->bNoBlendWeight = !(NewState == ESlateCheckBoxState::Checked);
+}
+
 FText STiledLandcapeImportDlg::GenerateConfigurationText(int32 NumComponents, int32 NumSectionsPerComponent, int32 NumQuadsPerSection) const
 {
 	const FString ComponentsStr = FString::Printf(TEXT("%dx%d"), NumComponents, NumComponents);
@@ -645,19 +675,13 @@ FText STiledLandcapeImportDlg::GenerateConfigurationText(int32 NumComponents, in
 
 void STiledLandcapeImportDlg::UpdateLandscapeLayerList()
 {
-	ImportSettings.LandscapeLayerNameList = ALandscapeProxy::GetLayersFromMaterial(ImportSettings.LandscapeMaterial.Get());
-	//
-	ImportSettings.WeightmapFileList.Empty();
-	ImportSettings.WeightmapFileList.SetNum(ImportSettings.LandscapeLayerNameList.Num());
-
-	//	
+	TArray<FName> LayerNamesList = ALandscapeProxy::GetLayersFromMaterial(ImportSettings.LandscapeMaterial.Get());
 	LayerDataList.Empty();
-	for (FName LayerName  : ImportSettings.LandscapeLayerNameList)
+	for (FName LayerName : LayerNamesList)
 	{
-		TSharedPtr<FLandscapeImportLayerData> LayerData = MakeShareable(new FLandscapeImportLayerData());
-		LayerData->LayerName = LayerName.ToString();
-		LayerData->bBlend = true;
-
+		// List view data source
+		TSharedPtr<FTiledLandscapeImportSettings::LandscapeLayerSettings> LayerData = MakeShareable(new FTiledLandscapeImportSettings::LandscapeLayerSettings());
+		LayerData->Name = LayerName;
 		LayerDataList.Add(LayerData);
 	}
 
diff --git a/Engine/Source/Editor/WorldBrowser/Private/Tiles/STiledLandscapeImportDlg.h b/Engine/Source/Editor/WorldBrowser/Private/Tiles/STiledLandscapeImportDlg.h
--- a/Engine/Source/Editor/WorldBrowser/Private/Tiles/STiledLandscapeImportDlg.h
+++ b/Engine/Source/Editor/WorldBrowser/Private/Tiles/STiledLandscapeImportDlg.h
@@ -18,13 +18,6 @@ public:
 		int32 NumQuadsPerSection;
 	};
 
-	struct FLandscapeImportLayerData
-	{
-		FString			LayerName;
-		TArray<FString> WeighmapFileList;
-		bool			bBlend;
-	};
-
 	SLATE_BEGIN_ARGS( STiledLandcapeImportDlg )
 		{}
 	SLATE_END_ARGS()
@@ -44,7 +37,7 @@ private:
 	FText				GetTileConfigurationText() const;
 	
 	/** */
-	TSharedRef<ITableRow> OnGenerateWidgetForLayerDataListView(TSharedPtr<FLandscapeImportLayerData> InLayerData, const TSharedRef<STableViewBase>& OwnerTable);
+	TSharedRef<ITableRow> OnGenerateWidgetForLayerDataListView(TSharedPtr<FTiledLandscapeImportSettings::LandscapeLayerSettings> InLayerData, const TSharedRef<STableViewBase>& OwnerTable);
 
 	/** */
 	TOptional<float> GetScaleX() const;
@@ -66,7 +59,7 @@ private:
 		
 	/** */
 	FReply OnClickedSelectHeightmapTiles();
-	FReply OnClickedSelectWeighmapTiles(TSharedPtr<FLandscapeImportLayerData> InLayerData);
+	FReply OnClickedSelectWeighmapTiles(TSharedPtr<FTiledLandscapeImportSettings::LandscapeLayerSettings> InLayerData);
 	FReply OnClickedImport();
 	FReply OnClickedCancel();
 
@@ -75,7 +68,11 @@ private:
 
 	/** */
 	FText GetImportSummaryText() const;
-	FText GetWeightmapCountText(TSharedPtr<FLandscapeImportLayerData> InLayerData) const; 
+	FText GetWeightmapCountText(TSharedPtr<FTiledLandscapeImportSettings::LandscapeLayerSettings> InLayerData) const; 
+
+	/** */
+	ESlateCheckBoxState::Type GetLayerBlendState(TSharedPtr<FTiledLandscapeImportSettings::LandscapeLayerSettings> InLayerData) const;
+	void OnLayerBlendStateChanged(ESlateCheckBoxState::Type NewState, TSharedPtr<FTiledLandscapeImportSettings::LandscapeLayerSettings> InLayerData);
 
 	/** */
 	void SetPossibleConfigurationsForResolution(int32 TargetResolutuion);
@@ -99,8 +96,8 @@ private:
 	/** */
 	TSharedPtr<SComboButton> LandscapeMaterialComboButton;
 
-	TSharedPtr<SListView<TSharedPtr<FLandscapeImportLayerData>>>	LayerDataListView;
-	TArray<TSharedPtr<FLandscapeImportLayerData>>					LayerDataList;		
+	TSharedPtr<SListView<TSharedPtr<FTiledLandscapeImportSettings::LandscapeLayerSettings>>>	LayerDataListView;
+	TArray<TSharedPtr<FTiledLandscapeImportSettings::LandscapeLayerSettings>>					LayerDataList;		
 		
 	/** */
 	FTiledLandscapeImportSettings ImportSettings;
diff --git a/Engine/Source/Editor/WorldBrowser/Private/Tiles/WorldTileCollectionModel.cpp b/Engine/Source/Editor/WorldBrowser/Private/Tiles/WorldTileCollectionModel.cpp
--- a/Engine/Source/Editor/WorldBrowser/Private/Tiles/WorldTileCollectionModel.cpp
+++ b/Engine/Source/Editor/WorldBrowser/Private/Tiles/WorldTileCollectionModel.cpp
@@ -1346,19 +1346,22 @@ void FWorldTileCollectionModel::ImportTiledLandscape_Executed()
 			TileImportSettings.LandscapeTransform.SetScale3D(TileScale);
 
 			// Setup weightmaps for each layer
-			for (int32 LayerIdx = 0; LayerIdx < ImportSettings.LandscapeLayerNameList.Num(); ++LayerIdx)
+			for (int32 LayerIdx = 0; LayerIdx < ImportSettings.LandscapeLayerSettingsList.Num(); ++LayerIdx)
 			{
-				FName LayerName = ImportSettings.LandscapeLayerNameList[LayerIdx];
+				FName LayerName = ImportSettings.LandscapeLayerSettingsList[LayerIdx].Name;
 				
 				TileImportSettings.ImportLayers.Add(FLandscapeImportLayerInfo(LayerName));
 				FLandscapeImportLayerInfo& LayerImportInfo = TileImportSettings.ImportLayers.Last();
-
-				if (ImportSettings.WeightmapFileList[LayerIdx].IsValidIndex(TileIndex))
+				
+				// Do we have a weightmap for this tile?
+				FIntPoint TileCoordinates = ImportSettings.TileCoordinates[TileIndex];
+				const FString* WeightmapFile = ImportSettings.LandscapeLayerSettingsList[LayerIdx].WeightmapFiles.Find(TileCoordinates);
+				if (WeightmapFile)
 				{
-					LayerImportInfo.SourceFilePath = ImportSettings.WeightmapFileList[LayerIdx][TileIndex];
+					LayerImportInfo.SourceFilePath = *WeightmapFile;
 					ReadRawFile(LayerImportInfo.LayerData, *LayerImportInfo.SourceFilePath, FILEREAD_Silent);
 					LayerImportInfo.LayerInfo = GetandscapeLayerInfoObject(LayerImportInfo.LayerName, GetWorld()->GetOutermost()->GetName());
-					LayerImportInfo.LayerInfo->bNoWeightBlend = false; //option ?
+					LayerImportInfo.LayerInfo->bNoWeightBlend = ImportSettings.LandscapeLayerSettingsList[LayerIdx].bNoBlendWeight;
 				}
 			}
 						
