World Browser: Do not redraw tiles while shaders still compiling

diff --git a/Engine/Source/Editor/WorldBrowser/Private/LevelCollectionModel.cpp b/Engine/Source/Editor/WorldBrowser/Private/LevelCollectionModel.cpp
--- a/Engine/Source/Editor/WorldBrowser/Private/LevelCollectionModel.cpp
+++ b/Engine/Source/Editor/WorldBrowser/Private/LevelCollectionModel.cpp
@@ -1648,6 +1648,13 @@ void FLevelCollectionModel::OnLevelRemovedFromWorld(ULevel* InLevel, UWorld* InW
 
 void FLevelCollectionModel::OnRedrawAllViewports()
 {
+	if (GShaderCompilingManager && GShaderCompilingManager->IsCompiling())
+	{
+		// Editor seems like still compiling shaders, do not request tiles redraw until all shaders complation is finished
+		// Basically redraw only on last event
+		return;
+	}
+	
 	RequestRedrawAllLevels();
 }
 
