World Browser: Added reimporting of a tiled landscape (heightmap, weightmaps)

diff --git a/Engine/Source/Editor/LandscapeEditor/Private/LandscapeEdMode.cpp b/Engine/Source/Editor/LandscapeEditor/Private/LandscapeEdMode.cpp
--- a/Engine/Source/Editor/LandscapeEditor/Private/LandscapeEdMode.cpp
+++ b/Engine/Source/Editor/LandscapeEditor/Private/LandscapeEdMode.cpp
@@ -2626,4 +2626,32 @@ ALandscape* FEdModeLandscape::ChangeComponentSetting(int32 NumComponentsX, int32
 	return Landscape;
 }
 
+bool LandscapeEditorUtils::SetHeightmapData(ALandscapeProxy* Landscape, const TArray<uint16>& Data)
+{
+	FIntRect ComponentsRect = Landscape->GetBoundingRect() + Landscape->LandscapeSectionOffset;
+		
+	if (Data.Num() == (1+ComponentsRect.Width())*(1+ComponentsRect.Height()))
+	{
+		FHeightmapAccessor<false> HeightmapAccessor(Landscape->GetLandscapeInfo());
+		HeightmapAccessor.SetData(ComponentsRect.Min.X, ComponentsRect.Min.Y, ComponentsRect.Max.X, ComponentsRect.Max.Y, Data.GetTypedData());
+		return true;
+	}
+
+	return false;
+}
+
+bool LandscapeEditorUtils::SetWeightmapData(ALandscapeProxy* Landscape, ULandscapeLayerInfoObject* LayerObject, const TArray<uint8>& Data)
+{
+	FIntRect ComponentsRect = Landscape->GetBoundingRect() + Landscape->LandscapeSectionOffset;	
+	
+	if (Data.Num() == (1+ComponentsRect.Width())*(1+ComponentsRect.Height()))
+	{
+		FAlphamapAccessor<false, true> AlphamapAccessor(Landscape->GetLandscapeInfo(), LayerObject);
+		AlphamapAccessor.SetData(ComponentsRect.Min.X, ComponentsRect.Min.Y, ComponentsRect.Max.X, ComponentsRect.Max.Y, Data.GetTypedData(), ELandscapeLayerPaintingRestriction::None);
+		return true;
+	}
+
+	return false;
+}
+
 #undef LOCTEXT_NAMESPACE
diff --git a/Engine/Source/Editor/LandscapeEditor/Public/LandscapeEdMode.h b/Engine/Source/Editor/LandscapeEditor/Public/LandscapeEdMode.h
--- a/Engine/Source/Editor/LandscapeEditor/Public/LandscapeEdMode.h
+++ b/Engine/Source/Editor/LandscapeEditor/Public/LandscapeEdMode.h
@@ -744,4 +744,8 @@ namespace LandscapeEditorUtils
 
 		return Result;
 	}
+
+	bool LANDSCAPEEDITOR_API SetHeightmapData(ALandscapeProxy* Landscape, const TArray<uint16>& Data);
+
+	bool LANDSCAPEEDITOR_API SetWeightmapData(ALandscapeProxy* Landscape, ULandscapeLayerInfoObject* LayerObject, const TArray<uint8>& Data);
 }
diff --git a/Engine/Source/Editor/WorldBrowser/Private/Tiles/WorldTileCollectionModel.cpp b/Engine/Source/Editor/WorldBrowser/Private/Tiles/WorldTileCollectionModel.cpp
--- a/Engine/Source/Editor/WorldBrowser/Private/Tiles/WorldTileCollectionModel.cpp
+++ b/Engine/Source/Editor/WorldBrowser/Private/Tiles/WorldTileCollectionModel.cpp
@@ -9,6 +9,7 @@
 #include "MaterialExportUtils.h"
 #include "MeshUtilities.h"
 #include "RawMesh.h"
+#include "LandscapeEdMode.h"
 
 #include "WorldTileDetails.h"
 #include "WorldTileDetailsCustomization.h"
@@ -17,6 +18,10 @@
 
 #define LOCTEXT_NAMESPACE "WorldBrowser"
 
+
+static const FName HeightmapLayerName = FName("__Heightmap__");
+
+
 FWorldTileCollectionModel::FWorldTileCollectionModel(const TWeakObjectPtr<UEditorEngine>& InEditor)
 	: FLevelCollectionModel(InEditor)
 	, bIsSavingLevel(false)
@@ -308,19 +313,6 @@ void FWorldTileCollectionModel::BuildGridMenu(FMenuBuilder& MenuBuilder) const
 			}
 			MenuBuilder.EndSection();
 		}
-
-		// Landscape specific commands
-		if (StaticCastSharedPtr<FWorldTileModel>(SelectedLevelsList[0])->IsLandscapeBased())
-		{
-			MenuBuilder.BeginSection("Menu_AddLandscape");
-			{
-				MenuBuilder.AddSubMenu( 
-					LOCTEXT("AddLandscapeLevel", "Add Adjacent Landscape Level"),
-					FText::GetEmpty(),
-					FNewMenuDelegate::CreateSP(this, &FWorldTileCollectionModel::BuildAdjacentLandscapeMenu));
-			}
-			MenuBuilder.EndSection();
-		}
 	}
 	else
 	{
@@ -331,6 +323,34 @@ void FWorldTileCollectionModel::BuildGridMenu(FMenuBuilder& MenuBuilder) const
 		MenuBuilder.EndSection();
 	}
 
+	// Landscape specific stuff
+	const bool bCanReimportTiledLandscape = CanReimportTiledlandscape();
+	const bool bCanAddAdjacentLandscape = CanAddLandscapeProxy(FWorldTileModel::EWorldDirections::Any);
+	if (bCanReimportTiledLandscape || bCanAddAdjacentLandscape)
+	{
+		MenuBuilder.BeginSection("Menu_LandscapeSection", LOCTEXT("Menu_LandscapeSectionTitle", "Landscape"));
+		
+		// Adjacent landscape
+		if (bCanAddAdjacentLandscape)
+		{
+			MenuBuilder.AddSubMenu( 
+				LOCTEXT("AddLandscapeLevel", "Add Adjacent Landscape Level"),
+				FText::GetEmpty(),
+				FNewMenuDelegate::CreateSP(this, &FWorldTileCollectionModel::BuildAdjacentLandscapeMenu));
+		}
+
+		// Tiled landscape
+		if (bCanReimportTiledLandscape)
+		{
+			MenuBuilder.AddSubMenu( 
+				LOCTEXT("ReimportTiledLandscape", "Reimport Tiled Landscape"),
+				FText::GetEmpty(),
+				FNewMenuDelegate::CreateSP(this, &FWorldTileCollectionModel::BuildReimportTiledLandscapeMenu));
+		}
+		
+		MenuBuilder.EndSection();
+	}
+
 	if (AreAnySelectedLevelsEditable())
 	{
 		MenuBuilder.BeginSection("WorldGridItemLayers");
@@ -420,6 +440,68 @@ void FWorldTileCollectionModel::BuildAdjacentLandscapeMenu(FMenuBuilder& MenuBui
 	MenuBuilder.AddMenuEntry(Commands.AddLandscapeLevelYPositive, NAME_None, TAttribute<FText>(), TAttribute<FText>(), FSlateIcon(FEditorStyle::GetStyleSetName(), "WorldBrowser.DirectionYPositive"));
 }
 
+void FWorldTileCollectionModel::BuildReimportTiledLandscapeMenu(FMenuBuilder& InMenuBuilder) const
+{
+	const FLevelCollectionCommands& Commands = FLevelCollectionCommands::Get();
+
+	// Add "Heightmap" menu entry
+	InMenuBuilder.AddMenuEntry(
+			LOCTEXT("Menu_HeightmapTitle", "Heightmap"), 
+			FText(), FSlateIcon(),
+			FUIAction(FExecuteAction::CreateSP(
+				this, &FWorldTileCollectionModel::ReimportTiledLandscape_Executed, HeightmapLayerName
+				)
+			)
+		);
+
+	// Weightmaps	
+	InMenuBuilder.AddSubMenu( 
+				LOCTEXT("Menu_WeightmapsTitle", "Weightmaps"),
+				FText::GetEmpty(),
+				FNewMenuDelegate::CreateSP(this, &FWorldTileCollectionModel::BuildWeightmapsMenu)
+			);
+}
+
+void FWorldTileCollectionModel::BuildWeightmapsMenu(FMenuBuilder& InMenuBuilder) const
+{
+	// Add "All Weighmaps" menu entry
+	InMenuBuilder.AddMenuEntry(
+			LOCTEXT("Menu_AllWeightmapsTitle", "All Weightmaps"), 
+			FText(), FSlateIcon(),
+			FUIAction(FExecuteAction::CreateSP(
+				this, &FWorldTileCollectionModel::ReimportTiledLandscape_Executed, FName(NAME_None)
+				)
+			)
+		);
+	
+	TArray<FName> LayerNames;
+	// Get list of the landscape layers
+	for (const auto& LevelModel : SelectedLevelsList)
+	{
+		auto TileModel = StaticCastSharedPtr<FWorldTileModel>(LevelModel);
+		if (TileModel->IsTiledLandscapeBased())
+		{
+			TArray<FName> Layers = ALandscapeProxy::GetLayersFromMaterial(TileModel->GetLandscape()->LandscapeMaterial);
+			for (FName LayerName : Layers)
+			{
+				LayerNames.AddUnique(LayerName);
+			}
+		}
+	}
+
+	for (FName LayerName : LayerNames)
+	{
+		InMenuBuilder.AddMenuEntry(
+			FText::FromName(LayerName), 
+			FText(), FSlateIcon(),
+			FUIAction(FExecuteAction::CreateSP(
+				this, &FWorldTileCollectionModel::ReimportTiledLandscape_Executed, LayerName
+				)
+			)
+		);
+	}
+}
+
 void FWorldTileCollectionModel::CustomizeFileMainMenu(FMenuBuilder& InMenuBuilder) const
 {
 	FLevelCollectionModel::CustomizeFileMainMenu(InMenuBuilder);
@@ -428,8 +510,8 @@ void FWorldTileCollectionModel::CustomizeFileMainMenu(FMenuBuilder& InMenuBuilde
 		
 	InMenuBuilder.BeginSection("LevelsAddLevel");
 	{
-		InMenuBuilder.AddMenuEntry( Commands.World_CreateEmptyLevel );
-		InMenuBuilder.AddMenuEntry( Commands.ImportTiledLandscape );
+		InMenuBuilder.AddMenuEntry(Commands.World_CreateEmptyLevel);
+		InMenuBuilder.AddMenuEntry(Commands.ImportTiledLandscape);
 	}
 	InMenuBuilder.EndSection();
 }
@@ -557,7 +639,7 @@ void FWorldTileCollectionModel::BindCommands()
 	// Landscape operations
 	ActionList.MapAction(Commands.ImportTiledLandscape,
 		FExecuteAction::CreateSP(this, &FWorldTileCollectionModel::ImportTiledLandscape_Executed));
-
+		
 	ActionList.MapAction(Commands.AddLandscapeLevelXNegative,
 		FExecuteAction::CreateSP(this, &FWorldTileCollectionModel::AddLandscapeProxy_Executed, FWorldTileModel::XNegative),
 		FCanExecuteAction::CreateSP(this, &FWorldTileCollectionModel::CanAddLandscapeProxy, FWorldTileModel::XNegative));
@@ -1084,7 +1166,7 @@ FVector2D FWorldTileCollectionModel::SnapTranslationDeltaLandscape(const TShared
 																	FVector2D InAbsoluteDelta, 
 																	float SnappingDistance)
 {
-	ALandscapeProxy* Landscape = LandscapeTile->GetLandcape();
+	ALandscapeProxy* Landscape = LandscapeTile->GetLandscape();
 	FVector ComponentScale = Landscape->GetRootComponent()->RelativeScale3D*Landscape->ComponentSizeQuads;
 	
 	return FVector2D(	FMath::GridSnap(InAbsoluteDelta.X, ComponentScale.X),
@@ -1203,6 +1285,19 @@ bool FWorldTileCollectionModel::CanAddLandscapeProxy(FWorldTileModel::EWorldDire
 	return false;
 }
 
+bool FWorldTileCollectionModel::CanReimportTiledlandscape() const
+{
+	for (const auto& LevelModel : SelectedLevelsList)
+	{
+		if (LevelModel->IsEditable() && StaticCastSharedPtr<FWorldTileModel>(LevelModel)->IsTiledLandscapeBased())
+		{
+			return true;
+		}
+	}
+
+	return false;
+}
+
 void FWorldTileCollectionModel::AddLandscapeProxy_Executed(FWorldTileModel::EWorldDirections InWhere)
 {
 	if (IsReadOnly() || !IsOneLevelSelected())
@@ -1228,7 +1323,7 @@ void FWorldTileCollectionModel::AddLandscapeProxy_Executed(FWorldTileModel::EWor
 		FLevelModelList Levels; 
 		Levels.Add(NewLevelModel);
 
-		ALandscapeProxy* SourceLandscape = LandscapeTileModel->GetLandcape();
+		ALandscapeProxy* SourceLandscape = LandscapeTileModel->GetLandscape();
 		FIntPoint SourceTileOffset = LandscapeTileModel->GetAbsoluteLevelPosition();
 
 		NewLevelModel->SetVisible(false);
@@ -1442,6 +1537,88 @@ void FWorldTileCollectionModel::ImportTiledLandscape_Executed()
 	}
 }
 
+void FWorldTileCollectionModel::ReimportTiledLandscape_Executed(FName TargetLayer)
+{
+	// Collect selected landscape tiles
+	TArray<TSharedPtr<FWorldTileModel>> TargetLandscapeTiles;
+	for (auto LevelModel : SelectedLevelsList)
+	{
+		TSharedPtr<FWorldTileModel> TileModel = StaticCastSharedPtr<FWorldTileModel>(LevelModel);
+		if (TileModel->IsEditable() && TileModel->IsTiledLandscapeBased())
+		{
+			TargetLandscapeTiles.Add(TileModel);
+		}
+	}
+
+	if (TargetLandscapeTiles.Num() == 0)
+	{
+		return;
+	}
+			
+	TArray<bool> AllLevelsVisibilityState;
+	// Hide all visible levels
+	for (auto LevelModel : AllLevelsList)
+	{
+		AllLevelsVisibilityState.Add(LevelModel->IsVisible());
+		if (!LevelModel->IsPersistent())
+		{
+			LevelModel->SetVisible(false);
+		}
+	}
+	
+	// Reimport data for each selected landscape tile
+	for (auto TileModel : TargetLandscapeTiles)
+	{
+		TileModel->SetVisible(true);
+
+		ALandscapeProxy* Landscape = TileModel->GetLandscape();
+
+		ULandscapeLayerInfoObject* DataLayer = ALandscapeProxy::DataLayer;
+	
+		if (TargetLayer == HeightmapLayerName) // Heightmap
+		{
+			if (!Landscape->ReimportHeightmapFilePath.IsEmpty())
+			{
+				TArray<uint16> RawData;
+				ReadRawFile(RawData, *Landscape->ReimportHeightmapFilePath, FILEREAD_Silent);
+				LandscapeEditorUtils::SetHeightmapData(Landscape, RawData);
+			}
+		}
+		else // Weightmap
+		{
+			for (FLandscapeEditorLayerSettings& LayerSettings : Landscape->EditorLayerSettings)
+			{
+				if (LayerSettings.LayerInfoObj && (LayerSettings.LayerInfoObj->LayerName == TargetLayer || TargetLayer == NAME_None))
+				{
+					if (!LayerSettings.ReimportLayerFilePath.IsEmpty())
+					{
+						TArray<uint8> RawData;
+						ReadRawFile(RawData, *LayerSettings.ReimportLayerFilePath, FILEREAD_Silent);
+						LandscapeEditorUtils::SetWeightmapData(Landscape, LayerSettings.LayerInfoObj, RawData);
+
+						if (TargetLayer != NAME_None)
+						{
+							// Importing one specific layer
+							break;
+						}
+					}
+				}
+			}
+		}
+		
+		TileModel->SetVisible(false);
+	}
+		
+	// Restore levels visibility
+	for (int32 LevelIdx = 0; LevelIdx < AllLevelsList.Num(); ++LevelIdx)
+	{
+		if (AllLevelsVisibilityState[LevelIdx])
+		{
+			AllLevelsList[LevelIdx]->SetVisible(true);
+		}
+	}
+}
+
 void FWorldTileCollectionModel::PostUndo(bool bSuccess)
 {
 	if (!bIsSavingLevel)
diff --git a/Engine/Source/Editor/WorldBrowser/Private/Tiles/WorldTileCollectionModel.h b/Engine/Source/Editor/WorldBrowser/Private/Tiles/WorldTileCollectionModel.h
--- a/Engine/Source/Editor/WorldBrowser/Private/Tiles/WorldTileCollectionModel.h
+++ b/Engine/Source/Editor/WorldBrowser/Private/Tiles/WorldTileCollectionModel.h
@@ -187,6 +187,12 @@ private:
 	/** Builds adjacent landscape menu */
 	void BuildAdjacentLandscapeMenu(FMenuBuilder& InMenuBuilder) const;
 
+	/** Builds reimport tiled landscape menu */
+	void BuildReimportTiledLandscapeMenu(FMenuBuilder& InMenuBuilder) const;
+	
+	/** Builds reimport weightmaps menu */
+	void BuildWeightmapsMenu(FMenuBuilder& InMenuBuilder) const;
+	
 private:
 	/** Creates a new empty Level; prompts for level save location */
 	void CreateEmptyLevel_Executed();
@@ -212,9 +218,15 @@ private:
 	/** @return whether it is possible to add a level with landscape proxy at specified location */	
 	bool CanAddLandscapeProxy(FWorldTileModel::EWorldDirections InWhere) const;
 
+	/** @return Whether selection contains tiles with tiled landscape */	
+	bool CanReimportTiledlandscape() const;
+
 	/**  */	
 	void ImportTiledLandscape_Executed();
 
+	/**  */	
+	void ReimportTiledLandscape_Executed(FName TargetLayer);
+
 public:
 	/** Whether Editor has support for generating LOD levels */	
 	bool HasGenerateLODLevelSupport() const;
diff --git a/Engine/Source/Editor/WorldBrowser/Private/Tiles/WorldTileModel.cpp b/Engine/Source/Editor/WorldBrowser/Private/Tiles/WorldTileModel.cpp
--- a/Engine/Source/Editor/WorldBrowser/Private/Tiles/WorldTileModel.cpp
+++ b/Engine/Source/Editor/WorldBrowser/Private/Tiles/WorldTileModel.cpp
@@ -307,6 +307,26 @@ bool FWorldTileModel::IsLandscapeBased() const
 	return Landscape.IsValid();
 }
 
+bool FWorldTileModel::IsTiledLandscapeBased() const
+{
+	if (IsLandscapeBased() && !GetLandscape()->ReimportHeightmapFilePath.IsEmpty())
+	{
+		// Check if single landscape actor resolution matches heightmap file size
+		IFileManager& FileManager = IFileManager::Get();
+		const int64 ImportFileSize = FileManager.FileSize(*GetLandscape()->ReimportHeightmapFilePath);
+		
+		FIntRect ComponentsRect = GetLandscape()->GetBoundingRect();
+		int64 LandscapeSamples	= (int64)(ComponentsRect.Width()+1)*(ComponentsRect.Height()+1);
+		// Height samples are 2 bytes wide
+		if (LandscapeSamples*2 == ImportFileSize)
+		{
+			return true;
+		}
+	}
+
+	return false;
+}
+
 bool FWorldTileModel::IsLandscapeProxy() const
 {
 	return (Landscape.IsValid() && !Landscape.Get()->IsA(ALandscape::StaticClass()));
@@ -322,7 +342,7 @@ bool FWorldTileModel::IsInLayersList(const TArray<FWorldTileLayer>& InLayerList)
 	return true;
 }
 
-ALandscapeProxy* FWorldTileModel::GetLandcape() const
+ALandscapeProxy* FWorldTileModel::GetLandscape() const
 {
 	return Landscape.Get();
 }
diff --git a/Engine/Source/Editor/WorldBrowser/Private/Tiles/WorldTileModel.h b/Engine/Source/Editor/WorldBrowser/Private/Tiles/WorldTileModel.h
--- a/Engine/Source/Editor/WorldBrowser/Private/Tiles/WorldTileModel.h
+++ b/Engine/Source/Editor/WorldBrowser/Private/Tiles/WorldTileModel.h
@@ -28,7 +28,8 @@ public:
 		XNegative,
 		YNegative,
 		XPositive,
-		YPositive
+		YPositive,
+		Any
 	};
 
 	/**
@@ -79,9 +80,6 @@ public:
 	virtual void OnLevelRemovedFromWorld() OVERRIDE;
 	virtual void OnParentChanged() OVERRIDE;
 	// FLevelModel interface end
-		
-	/** @return The landscape actor in case this level is landscape based */
-	ALandscapeProxy* GetLandcape() const;
 	
 	/** Adds new streaming level*/
 	void AddStreamingLevel(UClass* InStreamingClass, const FName& InPackageName);
@@ -106,9 +104,15 @@ public:
 
 	/** Whether this level landscape based or not */
 	bool IsLandscapeBased() const;
-	
+
+	/** Whether this level based on tiled landscape or not */
+	bool IsTiledLandscapeBased() const;
+		
 	/** Whether this level has ALandscapeProxy or not */
 	bool IsLandscapeProxy() const;
+
+	/** @return The landscape actor in case this level is landscape based */
+	ALandscapeProxy* GetLandscape() const;
 	
 	/** Whether this level in provided layers list */
 	bool IsInLayersList(const TArray<FWorldTileLayer>& InLayerList) const;
diff --git a/Engine/Source/Editor/WorldBrowser/WorldBrowser.Build.cs b/Engine/Source/Editor/WorldBrowser/WorldBrowser.Build.cs
--- a/Engine/Source/Editor/WorldBrowser/WorldBrowser.Build.cs
+++ b/Engine/Source/Editor/WorldBrowser/WorldBrowser.Build.cs
@@ -36,6 +36,7 @@ public class WorldBrowser : ModuleRules
 				"SourceControlWindows",
                 "RawMesh",
                 "MeshUtilities",
+                "LandscapeEditor",
             }
 		);
 
diff --git a/Engine/Source/Runtime/Engine/Private/Landscape/LandscapeEdit.cpp b/Engine/Source/Runtime/Engine/Private/Landscape/LandscapeEdit.cpp
--- a/Engine/Source/Runtime/Engine/Private/Landscape/LandscapeEdit.cpp
+++ b/Engine/Source/Runtime/Engine/Private/Landscape/LandscapeEdit.cpp
@@ -2383,7 +2383,7 @@ FIntRect ALandscapeProxy::GetBoundingRect() const
 		Rect.Include(LandscapeComponents[CompIdx]->GetSectionBase());
 	}
 	
-	if (Rect.Area() > 0)
+	if (Rect.Width() > 0 && Rect.Height() > 0)
 	{
 		Rect.Max+= FIntPoint(ComponentSizeQuads, ComponentSizeQuads);
 		Rect-= LandscapeSectionOffset;
