World Browser: Landscape uses maximum LOD level for exporting to a raw mesh during LOD maps generation

diff --git a/Engine/Source/Developer/MeshUtilities/Private/MeshUtilities.cpp b/Engine/Source/Developer/MeshUtilities/Private/MeshUtilities.cpp
--- a/Engine/Source/Developer/MeshUtilities/Private/MeshUtilities.cpp
+++ b/Engine/Source/Developer/MeshUtilities/Private/MeshUtilities.cpp
@@ -2622,7 +2622,7 @@ void FMeshUtilities::CreateProxyMesh(
 		TArray<int32> RawMeshMaterialMap;
 		int32 RawMeshId = RawMeshes.Add(FRawMesh());
 
-		if (Landscape->ExportToRawMesh(RawMeshes[RawMeshId]))
+		if (Landscape->ExportToRawMesh(INDEX_NONE, RawMeshes[RawMeshId]))
 		{
 			// Landscape has one unique material
 			int32 MatIdx = UniqueMaterials.Add(MaterialExportUtils::FFlattenMaterial());
diff --git a/Engine/Source/Editor/WorldBrowser/Private/Tiles/WorldTileCollectionModel.cpp b/Engine/Source/Editor/WorldBrowser/Private/Tiles/WorldTileCollectionModel.cpp
--- a/Engine/Source/Editor/WorldBrowser/Private/Tiles/WorldTileCollectionModel.cpp
+++ b/Engine/Source/Editor/WorldBrowser/Private/Tiles/WorldTileCollectionModel.cpp
@@ -1809,9 +1809,10 @@ bool FWorldTileCollectionModel::GenerateLODLevel(TSharedPtr<FLevelModel> InLevel
 		
 		FRawMesh LandscapeRawMesh;
 		FFlattenMaterial LandscapeFlattenMaterial;
-		
 		FVector LandscapeWorldLocation = Landscape->GetActorLocation();
-		Landscape->ExportToRawMesh(LandscapeRawMesh);
+		
+		Landscape->ExportToRawMesh(MAX_int32, LandscapeRawMesh);
+		
 		for (FVector& VertexPos : LandscapeRawMesh.VertexPositions)
 		{
 			VertexPos-= LandscapeWorldLocation;
@@ -1822,16 +1823,21 @@ bool FWorldTileCollectionModel::GenerateLODLevel(TSharedPtr<FLevelModel> InLevel
 		ExportMaterial(Landscape, LandscapeFlattenMaterial);
 		
 		// Reduce landscape mesh
-		FMeshReductionSettings Settings;
-		Settings.PercentTriangles = LODInfo.GenDetailsPercentage/100.f;
-		float OutMaxDeviation = 0.f;
-		FRawMesh LandscapeReducedMesh = LandscapeRawMesh;
-
-		MeshUtilities.GetMeshReductionInterface()->Reduce(
-			LandscapeReducedMesh,
-			OutMaxDeviation,
-			LandscapeRawMesh,
-			Settings);
+		if (LODInfo.GenDetailsPercentage != 100.f)
+		{
+			FMeshReductionSettings Settings;
+			Settings.PercentTriangles = LODInfo.GenDetailsPercentage/100.f;
+			float OutMaxDeviation = 0.f;
+			FRawMesh LandscapeReducedMesh;
+
+			MeshUtilities.GetMeshReductionInterface()->Reduce(
+				LandscapeReducedMesh,
+				OutMaxDeviation,
+				LandscapeRawMesh,
+				Settings);
+
+			LandscapeRawMesh = LandscapeReducedMesh;
+		}
 
 		FString PackageName = LODContentDir + TEXT("LandscapeStatic");
 		PackageName = MakeUniqueObjectName(NULL, UPackage::StaticClass(), *PackageName).ToString();
@@ -1864,7 +1870,7 @@ bool FWorldTileCollectionModel::GenerateLODLevel(TSharedPtr<FLevelModel> InLevel
 			SrcModel->BuildSettings.bRecomputeTangents = false;
 			SrcModel->BuildSettings.bRemoveDegenerates = false;
 			SrcModel->BuildSettings.bUseFullPrecisionUVs = false;
-			SrcModel->RawMeshBulkData->SaveRawMesh(LandscapeReducedMesh);
+			SrcModel->RawMeshBulkData->SaveRawMesh(LandscapeRawMesh);
 
 			//Assign the proxy material to the static mesh
 			StaticMesh->Materials.Add(StaticLandscapeMaterial);
diff --git a/Engine/Source/Runtime/Engine/Classes/Landscape/LandscapeProxy.h b/Engine/Source/Runtime/Engine/Classes/Landscape/LandscapeProxy.h
--- a/Engine/Source/Runtime/Engine/Classes/Landscape/LandscapeProxy.h
+++ b/Engine/Source/Runtime/Engine/Classes/Landscape/LandscapeProxy.h
@@ -418,10 +418,11 @@ public:
 	/**
 	 * Exports landscape into raw mesh
 	 * 
+	 * @param InExportLOD Landscape LOD level to use while exporting, INDEX_NONE will use ALanscapeProxy::ExportLOD settings
 	 * @param OutRawMesh - Resulting raw mesh
 	 * @return true if successful
 	 */
-	ENGINE_API bool ExportToRawMesh(struct FRawMesh& OutRawMesh) const;
+	ENGINE_API bool ExportToRawMesh(int32 InExportLOD, struct FRawMesh& OutRawMesh) const;
 
 
 	/** @return Current size of bounding rectangle in quads space */
diff --git a/Engine/Source/Runtime/Engine/Private/Landscape/LandscapeEdit.cpp b/Engine/Source/Runtime/Engine/Private/Landscape/LandscapeEdit.cpp
--- a/Engine/Source/Runtime/Engine/Private/Landscape/LandscapeEdit.cpp
+++ b/Engine/Source/Runtime/Engine/Private/Landscape/LandscapeEdit.cpp
@@ -2247,21 +2247,28 @@ void ALandscapeProxy::Import(FGuid Guid, int32 VertsX, int32 VertsY,
 	GWarn->EndSlowTask();
 }
 
-bool ALandscapeProxy::ExportToRawMesh(FRawMesh& OutRawMesh) const
+bool ALandscapeProxy::ExportToRawMesh(int32 InExportLOD, FRawMesh& OutRawMesh) const
 {
 	TArray<ULandscapeComponent*> RegisteredLandscapeComponents;
 	GetComponents<ULandscapeComponent>(RegisteredLandscapeComponents);
 	
 	const FIntRect LandscapeSectionRect = GetBoundingRect();
 	const FVector2D LandscapeUVScale = FVector2D(1.f, 1.f)/FVector2D(LandscapeSectionRect.Size());
+
+	// User specified LOD to export
+	int32 LandscapeLODToExport = ExportLOD; 
+	if (InExportLOD != INDEX_NONE)
+	{
+		LandscapeLODToExport = FMath::Clamp<int32>(InExportLOD, 0, FMath::CeilLogTwo(SubsectionSizeQuads+1)-1);
+	}
 		
 	// Export data for each component
 	for (auto It = RegisteredLandscapeComponents.CreateConstIterator(); It; ++It)
 	{
 		ULandscapeComponent* Component = (*It);
-		FLandscapeComponentDataInterface CDI(Component, ExportLOD);
-		const int32 ComponentSizeQuadsLOD = ((Component->ComponentSizeQuads+1)>>ExportLOD)-1;
-		const int32 SubsectionSizeQuadsLOD = ((Component->SubsectionSizeQuads+1)>>ExportLOD)-1;
+		FLandscapeComponentDataInterface CDI(Component, LandscapeLODToExport);
+		const int32 ComponentSizeQuadsLOD = ((Component->ComponentSizeQuads+1)>>LandscapeLODToExport)-1;
+		const int32 SubsectionSizeQuadsLOD = ((Component->SubsectionSizeQuads+1)>>LandscapeLODToExport)-1;
 		const FIntPoint ComponentOffsetQuads = Component->GetSectionBase() - LandscapeSectionOffset - LandscapeSectionRect.Min;
 		const FVector2D ComponentUVOffsetLOD = FVector2D(ComponentOffsetQuads)*((float)ComponentSizeQuadsLOD/ComponentSizeQuads);
 		const FVector2D ComponentUVScaleLOD = LandscapeUVScale*((float)ComponentSizeQuads/ComponentSizeQuadsLOD);
