World Browser: Added ability to generate LOD maps for hidden levels

diff --git a/Engine/Source/Editor/WorldBrowser/Private/Tiles/WorldTileCollectionModel.cpp b/Engine/Source/Editor/WorldBrowser/Private/Tiles/WorldTileCollectionModel.cpp
--- a/Engine/Source/Editor/WorldBrowser/Private/Tiles/WorldTileCollectionModel.cpp
+++ b/Engine/Source/Editor/WorldBrowser/Private/Tiles/WorldTileCollectionModel.cpp
@@ -1565,12 +1565,15 @@ void FWorldTileCollectionModel::ReimportTiledLandscape_Executed(FName TargetLaye
 			LevelModel->SetVisible(false);
 		}
 	}
-	
+
+	// Disable world origin tracking, so we can show, hide levels without offseting them
+	GetWorld()->WorldComposition->bTemporallyDisableOriginTracking = true;
+		
 	// Reimport data for each selected landscape tile
 	for (auto TileModel : TargetLandscapeTiles)
 	{
 		TileModel->SetVisible(true);
-
+			
 		ALandscapeProxy* Landscape = TileModel->GetLandscape();
 
 		ULandscapeLayerInfoObject* DataLayer = ALandscapeProxy::DataLayer;
@@ -1608,7 +1611,10 @@ void FWorldTileCollectionModel::ReimportTiledLandscape_Executed(FName TargetLaye
 		
 		TileModel->SetVisible(false);
 	}
-		
+
+	// Restore world origin tracking
+	GetWorld()->WorldComposition->bTemporallyDisableOriginTracking = false;
+			
 	// Restore levels visibility
 	for (int32 LevelIdx = 0; LevelIdx < AllLevelsList.Num(); ++LevelIdx)
 	{
@@ -1727,6 +1733,14 @@ bool FWorldTileCollectionModel::GenerateLODLevel(TSharedPtr<FLevelModel> InLevel
 		return false;
 	}
 
+	// We have to make original level visible, to correctly export it
+	const bool bVisibleLevel = InLevelModel->IsVisible();
+	if (!bVisibleLevel)
+	{
+		GetWorld()->WorldComposition->bTemporallyDisableOriginTracking = true;
+		InLevelModel->SetVisible(true);
+	}
+
 	IMeshUtilities& MeshUtilities = FModuleManager::Get().LoadModuleChecked<IMeshUtilities>("MeshUtilities");
 	IMeshMerging* MeshMerging = MeshUtilities.GetMeshMergingInterface();
 	check(MeshMerging);
@@ -1752,6 +1766,10 @@ bool FWorldTileCollectionModel::GenerateLODLevel(TSharedPtr<FLevelModel> InLevel
 		
 	// This is current actors offset from their original position
 	FVector ActorsOffset = FVector(TileModel->GetAbsoluteLevelPosition() - CurrentWorld->GlobalOriginOffset);
+	if (GetWorld()->WorldComposition->bTemporallyDisableOriginTracking)
+	{
+		ActorsOffset = FVector::ZeroVector;
+	}
 	
 	TArray<UObject*>			GeneratedAssets;
 	TArray<UStaticMesh*>		AssetsToSpawn;
@@ -1887,6 +1905,13 @@ bool FWorldTileCollectionModel::GenerateLODLevel(TSharedPtr<FLevelModel> InLevel
 		LandscapeActorIndex++;
 	}
 
+	// Restore level original visibility
+	if (!bVisibleLevel)
+	{
+		InLevelModel->SetVisible(false);
+		GetWorld()->WorldComposition->bTemporallyDisableOriginTracking = false;
+	}
+
 	// Save generated assets
 	TArray<UPackage*> AssetsPackagesToSave;
 	for (UObject* Asset : GeneratedAssets)
diff --git a/Engine/Source/Editor/WorldBrowser/Private/Tiles/WorldTileModel.cpp b/Engine/Source/Editor/WorldBrowser/Private/Tiles/WorldTileModel.cpp
--- a/Engine/Source/Editor/WorldBrowser/Private/Tiles/WorldTileModel.cpp
+++ b/Engine/Source/Editor/WorldBrowser/Private/Tiles/WorldTileModel.cpp
@@ -206,6 +206,12 @@ bool FWorldTileModel::ShouldBeVisible(FBox EditableArea) const
 		return true;
 	}
 
+	// When this hack is activated level should be visible regardless of current world origin
+	if (LevelCollectionModel.GetWorld()->WorldComposition->bTemporallyDisableOriginTracking)
+	{
+		return true;
+	}
+	
 	FBox LevelBBox = GetLevelBounds();
 
 	// Visible if level has no valid bounds
diff --git a/Engine/Source/Runtime/Engine/Classes/Engine/WorldComposition.h b/Engine/Source/Runtime/Engine/Classes/Engine/WorldComposition.h
--- a/Engine/Source/Runtime/Engine/Classes/Engine/WorldComposition.h
+++ b/Engine/Source/Runtime/Engine/Classes/Engine/WorldComposition.h
@@ -170,6 +170,10 @@ public:
 #if WITH_EDITOR
 	// Last location from where streaming state was updated
 	mutable FVector				LastViewLocation;
+
+	// Hack for a World Browser to be able to temporally show hidden levels 
+	// regardless of current world origin and without offsetting them temporally 
+	bool						bTemporallyDisableOriginTracking;
 #endif //WITH_EDITOR
 
 private:
diff --git a/Engine/Source/Runtime/Engine/Private/WorldComposition.cpp b/Engine/Source/Runtime/Engine/Private/WorldComposition.cpp
--- a/Engine/Source/Runtime/Engine/Private/WorldComposition.cpp
+++ b/Engine/Source/Runtime/Engine/Private/WorldComposition.cpp
@@ -582,6 +582,13 @@ void UWorldComposition::CommitTileStreamingState(UWorld* PersistenWorld, int32 T
 
 void UWorldComposition::OnLevelAddedToWorld(ULevel* InLevel)
 {
+#if WITH_EDITOR
+	if (bTemporallyDisableOriginTracking)
+	{
+		return;
+	}
+#endif
+		
 	// Move level according to current global origin offset
 	FIntPoint LevelOffset = GetLevelOffset(InLevel);
 	InLevel->ApplyWorldOffset(FVector(LevelOffset, 0), false);
@@ -589,6 +596,13 @@ void UWorldComposition::OnLevelAddedToWorld(ULevel* InLevel)
 
 void UWorldComposition::OnLevelRemovedFromWorld(ULevel* InLevel)
 {
+#if WITH_EDITOR
+	if (bTemporallyDisableOriginTracking)
+	{
+		return;
+	}
+#endif
+	
 	// Move level to his local origin
 	FIntPoint LevelOffset = GetLevelOffset(InLevel);
 	InLevel->ApplyWorldOffset(-FVector(LevelOffset, 0), false);
