World Browser: Fixed composition view scrolling when world orign has been shifted

diff --git a/Engine/Source/Editor/WorldBrowser/Private/Tiles/SWorldComposition.cpp b/Engine/Source/Editor/WorldBrowser/Private/Tiles/SWorldComposition.cpp
--- a/Engine/Source/Editor/WorldBrowser/Private/Tiles/SWorldComposition.cpp
+++ b/Engine/Source/Editor/WorldBrowser/Private/Tiles/SWorldComposition.cpp
@@ -72,7 +72,8 @@ public:
 
 	SWorldCompositionGrid()
 		: CommandList(MakeShareable(new FUICommandList))
-		, bHasScrollRequest(false)
+		, bHasScrollToRequest(false)
+		, bHasScrollByRequest(false)
 		, bIsFirstTickCall(true)
 		, bHasNodeInteraction(true)
 		, BoundsSnappingDistance(20.f)
@@ -95,6 +96,8 @@ public:
 		
 		WorldModel->SelectionChanged.RemoveAll(this);
 		WorldModel->CollectionChanged.RemoveAll(this);
+
+		FCoreDelegates::PreWorldOriginOffset.RemoveAll(this);
 	}
 
 	void Construct(const FArguments& InArgs)
@@ -136,6 +139,8 @@ public:
 		WorldModel->SelectionChanged.AddSP(this, &SWorldCompositionGrid::OnUpdateSelection);
 		WorldModel->CollectionChanged.AddSP(this, &SWorldCompositionGrid::RefreshView);
 		SelectionManager.OnSelectionChanged.BindSP(this, &SWorldCompositionGrid::OnSelectionChanged);
+
+		FCoreDelegates::PreWorldOriginOffset.AddSP(this, &SWorldCompositionGrid::PreWorldOriginOffset);
 	
 		RefreshView();
 	}
@@ -202,33 +207,48 @@ public:
 		//WorldModel->UpdateStreamingPreview(WorldMouseLocation, bShowPotentiallyVisibleLevels);
 			
 		// deffered scroll and zooming requests
-		if (bHasScrollRequest)
+		if (bHasScrollToRequest || bHasScrollByRequest)
 		{
-			bHasScrollRequest = false;
-
 			// zoom to
-			FVector2D SizeWithZoom = RequestedZoomArea*ZoomLevels->GetZoomAmount(ZoomLevel);
-			if (RequestedAllowZoomIn || 
-				(SizeWithZoom.X >= AllottedGeometry.Size.X || SizeWithZoom.Y >= AllottedGeometry.Size.Y))
+			if (RequestedAllowZoomIn)
 			{
-				// maximum zoom out by default
-				ZoomLevel = ZoomLevels->GetDefaultZoomLevel();
-				// expand zoom area little bit, so zooming will fit original area not so tight
-				RequestedZoomArea*= 1.2f;
-				// find more suitable zoom value
-				for (int32 Zoom = 0; Zoom < ZoomLevels->GetDefaultZoomLevel(); ++Zoom)
+				RequestedAllowZoomIn = false;
+				
+				FVector2D SizeWithZoom = RequestedZoomArea*ZoomLevels->GetZoomAmount(ZoomLevel);
+				
+				if (SizeWithZoom.X >= AllottedGeometry.Size.X || 
+					SizeWithZoom.Y >= AllottedGeometry.Size.Y)
 				{
-					SizeWithZoom = RequestedZoomArea*ZoomLevels->GetZoomAmount(Zoom);
-					if (SizeWithZoom.X >= AllottedGeometry.Size.X || SizeWithZoom.Y >= AllottedGeometry.Size.Y)
+					// maximum zoom out by default
+					ZoomLevel = ZoomLevels->GetDefaultZoomLevel();
+					// expand zoom area little bit, so zooming will fit original area not so tight
+					RequestedZoomArea*= 1.2f;
+					// find more suitable zoom value
+					for (int32 Zoom = 0; Zoom < ZoomLevels->GetDefaultZoomLevel(); ++Zoom)
 					{
-						ZoomLevel = Zoom;
-						break;
+						SizeWithZoom = RequestedZoomArea*ZoomLevels->GetZoomAmount(Zoom);
+						if (SizeWithZoom.X >= AllottedGeometry.Size.X || SizeWithZoom.Y >= AllottedGeometry.Size.Y)
+						{
+							ZoomLevel = Zoom;
+							break;
+						}
 					}
 				}
 			}
 
 			// scroll to
-			ViewOffset = RequestedScrollLocation - AllottedGeometry.Size*0.5f/GetZoomAmount();
+			if (bHasScrollToRequest)
+			{
+				bHasScrollToRequest = false;
+				ViewOffset = RequestedScrollToValue - AllottedGeometry.Size*0.5f/GetZoomAmount();
+			}
+
+			// scroll by
+			if (bHasScrollByRequest)
+			{
+				bHasScrollByRequest = false;
+				ViewOffset+= RequestedScrollByValue;
+			}
 		}
 	}
 	
@@ -737,12 +757,24 @@ protected:
 			FVector2D MinCorner, MaxCorner;
 			if (GetBoundsForNodes(true, MinCorner, MaxCorner, 0.f))
 			{
-				RequestScrollTo((MaxCorner + MinCorner)*0.5f, MaxCorner - MinCorner);
+				FVector2D TargetPosition = MaxCorner/2.f + MinCorner/2.f;
+				RequestScrollTo(TargetPosition, MaxCorner - MinCorner);
 			}
 		}
 		bUpdatingSelection = false;
 	}
 
+	/** Delegate callback: world origin is going to be moved. */
+	void PreWorldOriginOffset(UWorld* InWorld, const FIntPoint& InSrcOrigin, const FIntPoint& InDstOrigin)
+	{
+		if (WorldModel->GetWorld() != InWorld)
+		{
+			return;
+		}
+				
+		RequestScrollBy(-FVector2D(InDstOrigin - InSrcOrigin));
+	}
+	
 	/** Handles new item added to data source */
 	void OnNewItemAdded(TSharedPtr<FLevelModel> NewItem)
 	{
@@ -772,11 +804,17 @@ protected:
 	 */
 	void RequestScrollTo(FVector2D InLocation, FVector2D InArea, bool bAllowZoomIn = false)
 	{
-		bHasScrollRequest = true;
-		RequestedScrollLocation = InLocation;
+		bHasScrollToRequest = true;
+		RequestedScrollToValue = InLocation;
 		RequestedZoomArea = InArea;
 		RequestedAllowZoomIn = bAllowZoomIn;
 	}
+
+	void RequestScrollBy(FVector2D InDelta)
+	{
+		bHasScrollByRequest = true;
+		RequestedScrollByValue = InDelta;
+	}
 	
 	/** Handlers for moving items using arrow keys */
 	void MoveLevelLeft_Executed()
@@ -856,8 +894,10 @@ private:
 	TArray<FIntRect>						OccupiedCells;
 	const TSharedRef<FUICommandList>		CommandList;
 
-	bool									bHasScrollRequest;
-	FVector2D								RequestedScrollLocation;
+	bool									bHasScrollToRequest;
+	bool									bHasScrollByRequest;
+	FVector2D								RequestedScrollToValue;
+	FVector2D								RequestedScrollByValue;
 	FVector2D								RequestedZoomArea;
 	bool									RequestedAllowZoomIn;
 
