World Browser: Disabled landscape mesh simplificaton during LOD maps generation

diff --git a/Engine/Source/Editor/WorldBrowser/Private/Tiles/WorldTileCollectionModel.cpp b/Engine/Source/Editor/WorldBrowser/Private/Tiles/WorldTileCollectionModel.cpp
--- a/Engine/Source/Editor/WorldBrowser/Private/Tiles/WorldTileCollectionModel.cpp
+++ b/Engine/Source/Editor/WorldBrowser/Private/Tiles/WorldTileCollectionModel.cpp
@@ -1840,22 +1840,23 @@ bool FWorldTileCollectionModel::GenerateLODLevel(TSharedPtr<FLevelModel> InLevel
 		LandscapeFlattenMaterial.DiffuseSize = FIntPoint(1024, 1024);
 		ExportMaterial(Landscape, LandscapeFlattenMaterial);
 		
+		// TODO: Disabled landscape mesh simplification, does not really needed since exporting highest landscape LOD already provide low triangle count
 		// Reduce landscape mesh
-		if (LODInfo.GenDetailsPercentage != 100.f)
-		{
-			FMeshReductionSettings Settings;
-			Settings.PercentTriangles = LODInfo.GenDetailsPercentage/100.f;
-			float OutMaxDeviation = 0.f;
-			FRawMesh LandscapeReducedMesh;
-
-			MeshUtilities.GetMeshReductionInterface()->Reduce(
-				LandscapeReducedMesh,
-				OutMaxDeviation,
-				LandscapeRawMesh,
-				Settings);
-
-			LandscapeRawMesh = LandscapeReducedMesh;
-		}
+		//if (LODInfo.GenDetailsPercentage != 100.f)
+		//{
+		//	FMeshReductionSettings Settings;
+		//	Settings.PercentTriangles = LODInfo.GenDetailsPercentage/100.f;
+		//	float OutMaxDeviation = 0.f;
+		//	FRawMesh LandscapeReducedMesh;
+
+		//	MeshUtilities.GetMeshReductionInterface()->Reduce(
+		//		LandscapeReducedMesh,
+		//		OutMaxDeviation,
+		//		LandscapeRawMesh,
+		//		Settings);
+
+		//	LandscapeRawMesh = LandscapeReducedMesh;
+		//}
 
 		FString PackageName = LODContentDir + TEXT("LandscapeStatic");
 		PackageName = MakeUniqueObjectName(NULL, UPackage::StaticClass(), *PackageName).ToString();
