World Browser: Fixed case when bounding boxes were incorrectly calculated for hidden tiles

diff --git a/Engine/Source/Editor/WorldBrowser/Private/Tiles/WorldTileModel.cpp b/Engine/Source/Editor/WorldBrowser/Private/Tiles/WorldTileModel.cpp
--- a/Engine/Source/Editor/WorldBrowser/Private/Tiles/WorldTileModel.cpp
+++ b/Engine/Source/Editor/WorldBrowser/Private/Tiles/WorldTileModel.cpp
@@ -525,33 +525,26 @@ void FWorldTileModel::Update()
 		TileDetails->SyncStreamingLevels(*this);
 	
 		ULevel* Level = GetLevelObject();
-		if (Level == NULL)
-		{
-
-		}
-		else 
+		if (Level != nullptr && Level->bIsVisible)
 		{
 			if (Level->LevelBoundsActor.IsValid())
 			{
 				TileDetails->Bounds = Level->LevelBoundsActor.Get()->GetComponentsBoundingBox();
 			}
-
-			if (Level->bIsVisible)
+								
+			// True level bounds without offsets applied
+			if (TileDetails->Bounds.IsValid)
 			{
-				// True level bounds without offsets applied
-				if (TileDetails->Bounds.IsValid)
-				{
-					FBox LevelWorldBounds = TileDetails->Bounds;
-					FIntPoint GlobalOriginOffset = LevelCollectionModel.GetWorld()->GlobalOriginOffset;
-					FIntPoint LevelAbolutePosition = GetAbsoluteLevelPosition();
-					FIntPoint LevelOffset = LevelAbolutePosition - GlobalOriginOffset;
+				FBox LevelWorldBounds = TileDetails->Bounds;
+				FIntPoint GlobalOriginOffset = LevelCollectionModel.GetWorld()->GlobalOriginOffset;
+				FIntPoint LevelAbolutePosition = GetAbsoluteLevelPosition();
+				FIntPoint LevelOffset = LevelAbolutePosition - GlobalOriginOffset;
 
-					TileDetails->Bounds = LevelWorldBounds.ShiftBy(-FVector(LevelOffset, 0.f));
-				}
-
-				OnLevelInfoUpdated();
+				TileDetails->Bounds = LevelWorldBounds.ShiftBy(-FVector(LevelOffset, 0.f));
 			}
-		
+
+			OnLevelInfoUpdated();
+
 			// Cache landscape information
 			for (int32 ActorIndex = 0; ActorIndex < Level->Actors.Num() ; ++ActorIndex)
 			{
