World Browser: Renamed Kismet type of streaming levels to Blueprint

diff --git a/Engine/Source/Editor/WorldBrowser/Private/Tiles/WorldTileDetails.cpp b/Engine/Source/Editor/WorldBrowser/Private/Tiles/WorldTileDetails.cpp
--- a/Engine/Source/Editor/WorldBrowser/Private/Tiles/WorldTileDetails.cpp
+++ b/Engine/Source/Editor/WorldBrowser/Private/Tiles/WorldTileDetails.cpp
@@ -15,7 +15,7 @@ UClass* FTileStreamingLevelDetails::StreamingMode2Class(EStreamingLevelMode::Typ
 	{
 	case EStreamingLevelMode::AlwaysLoaded:
 		return ULevelStreamingAlwaysLoaded::StaticClass();
-	case EStreamingLevelMode::Kismet:
+	case EStreamingLevelMode::Blueprint:
 		return ULevelStreamingKismet::StaticClass();
 	case EStreamingLevelMode::Bounds:
 		return ULevelStreamingBounds::StaticClass();
@@ -34,7 +34,7 @@ EStreamingLevelMode::Type FTileStreamingLevelDetails::StreamingObject2Mode(ULeve
 
 	if (InLevelStreaming->IsA(ULevelStreamingKismet::StaticClass()))
 	{
-		return EStreamingLevelMode::Kismet; 
+		return EStreamingLevelMode::Blueprint; 
 	}
 
 	if (InLevelStreaming->IsA(ULevelStreamingBounds::StaticClass()))
diff --git a/Engine/Source/Editor/WorldBrowser/Private/Tiles/WorldTileDetails.h b/Engine/Source/Editor/WorldBrowser/Private/Tiles/WorldTileDetails.h
--- a/Engine/Source/Editor/WorldBrowser/Private/Tiles/WorldTileDetails.h
+++ b/Engine/Source/Editor/WorldBrowser/Private/Tiles/WorldTileDetails.h
@@ -12,7 +12,7 @@ namespace EStreamingLevelMode
 	enum Type
 	{
 		AlwaysLoaded,
-		Kismet,
+		Blueprint,
 		Bounds,
 	};
 }
