World Browser: Fixed 'white' tile images while simultating in the Editor

diff --git a/Engine/Source/Editor/WorldBrowser/Private/Tiles/SWorldTileItem.cpp b/Engine/Source/Editor/WorldBrowser/Private/Tiles/SWorldTileItem.cpp
--- a/Engine/Source/Editor/WorldBrowser/Private/Tiles/SWorldTileItem.cpp
+++ b/Engine/Source/Editor/WorldBrowser/Private/Tiles/SWorldTileItem.cpp
@@ -83,7 +83,7 @@ void FTileItemThumbnail::UpdateThumbnail()
 	}
 	
 	// Load image from a package header
-	if (!LevelModel->IsVisible())
+	if (!LevelModel->IsVisible() || LevelModel->IsSimulating())
 	{
 		TSet<FName> ObjectFullNames;
 		ObjectFullNames.Add(LevelModel->GetAssetName());
