World Browser: Removed inner streaming levels

diff --git a/Engine/Source/Editor/WorldBrowser/Private/LevelCollectionModel.cpp b/Engine/Source/Editor/WorldBrowser/Private/LevelCollectionModel.cpp
--- a/Engine/Source/Editor/WorldBrowser/Private/LevelCollectionModel.cpp
+++ b/Engine/Source/Editor/WorldBrowser/Private/LevelCollectionModel.cpp
@@ -240,24 +240,19 @@ void FLevelCollectionModel::Tick( float DeltaTime )
 		{
 			LevelModel->UpdateSimulationStatus(nullptr);
 		}
-		
-		// Traverse streaming levels in all inner worlds and update simulation status for corresponding level models  
-		for (ULevel* Level : GetSimulationWorld()->GetLevels())
+
+		// Traverse streaming levels and update simulation status for corresponding level models
+		for (ULevelStreaming* StreamingLevel : GetSimulationWorld()->StreamingLevels)
 		{
-			UWorld* InnerWorld = CastChecked<UWorld>(Level->GetOuter());
-			
-			for (ULevelStreaming* StreamingLevel : InnerWorld->StreamingLevels)
-			{
-				// Rebuild the original NonPrefixedPackageName so we can find it
-				const FString PrefixedPackageName = StreamingLevel->PackageName.ToString();
-				const FString NonPrefixedPackageName = FPackageName::GetLongPackagePath(PrefixedPackageName) + "/" 
-						+ FPackageName::GetLongPackageAssetName(PrefixedPackageName).RightChop(GetSimulationWorld()->StreamingLevelsPrefix.Len());
+			// Rebuild the original NonPrefixedPackageName so we can find it
+			const FString PrefixedPackageName = StreamingLevel->PackageName.ToString();
+			const FString NonPrefixedPackageName = FPackageName::GetLongPackagePath(PrefixedPackageName) + "/" 
+					+ FPackageName::GetLongPackageAssetName(PrefixedPackageName).RightChop(GetSimulationWorld()->StreamingLevelsPrefix.Len());
 								
-				TSharedPtr<FLevelModel> LevelModel = FindLevelModel(FName(*NonPrefixedPackageName));
-				if (LevelModel.IsValid())
-				{
-					LevelModel->UpdateSimulationStatus(StreamingLevel);
-				}
+			TSharedPtr<FLevelModel> LevelModel = FindLevelModel(FName(*NonPrefixedPackageName));
+			if (LevelModel.IsValid())
+			{
+				LevelModel->UpdateSimulationStatus(StreamingLevel);
 			}
 		}
 	}
diff --git a/Engine/Source/Editor/WorldBrowser/Private/Tiles/WorldTileDetails.cpp b/Engine/Source/Editor/WorldBrowser/Private/Tiles/WorldTileDetails.cpp
--- a/Engine/Source/Editor/WorldBrowser/Private/Tiles/WorldTileDetails.cpp
+++ b/Engine/Source/Editor/WorldBrowser/Private/Tiles/WorldTileDetails.cpp
@@ -6,45 +6,6 @@
 
 #define LOCTEXT_NAMESPACE "WorldBrowser"
 
-/////////////////////////////////////////////////////
-// FStreamingLevelDetails
-
-UClass* FTileStreamingLevelDetails::StreamingMode2Class(EStreamingLevelMode::Type InMode)
-{
-	switch (InMode)
-	{
-	case EStreamingLevelMode::AlwaysLoaded:
-		return ULevelStreamingAlwaysLoaded::StaticClass();
-	case EStreamingLevelMode::Blueprint:
-		return ULevelStreamingKismet::StaticClass();
-	case EStreamingLevelMode::Bounds:
-		return ULevelStreamingBounds::StaticClass();
-	default:
-		check(false); // unsupported mode
-	}
-	return NULL;
-}
-
-EStreamingLevelMode::Type FTileStreamingLevelDetails::StreamingObject2Mode(ULevelStreaming* InLevelStreaming)
-{
-	if (InLevelStreaming->IsA(ULevelStreamingAlwaysLoaded::StaticClass()))
-	{
-		return EStreamingLevelMode::AlwaysLoaded; 
-	}
-
-	if (InLevelStreaming->IsA(ULevelStreamingKismet::StaticClass()))
-	{
-		return EStreamingLevelMode::Blueprint; 
-	}
-
-	if (InLevelStreaming->IsA(ULevelStreamingBounds::StaticClass()))
-	{
-		return EStreamingLevelMode::Bounds; 
-	}
-
-	check(false); // unsupported class
-	return EStreamingLevelMode::AlwaysLoaded;
-}
 
 /////////////////////////////////////////////////////
 // FTileLODEntryDetails
@@ -122,29 +83,6 @@ FWorldTileInfo UWorldTileDetails::GetInfo() const
 	return Info;
 }
 
-void UWorldTileDetails::SyncStreamingLevels(const FWorldTileModel& TileModel)
-{
-	// Clean list if tile streaming details, 
-	// will be refiled with an actual information in a loop bellow
-	StreamingLevels.Empty();
-	
-	const ULevel* Level = TileModel.GetLevelObject();
-	if (Level)
-	{
-		auto WorldStreamingLevels = CastChecked<UWorld>(Level->GetOuter())->StreamingLevels;
-		for (auto It = WorldStreamingLevels.CreateConstIterator(); It; ++It)
-		{
-			ULevelStreaming* LevelStreaming = (*It);
-			
-			FTileStreamingLevelDetails Details;
-			Details.StreamingMode	= FTileStreamingLevelDetails::StreamingObject2Mode(LevelStreaming);
-			Details.PackageName		= LevelStreaming->PackageName;
-
-			StreamingLevels.Add(Details);
-		}
-	}
-}
-
 void UWorldTileDetails::PostEditChangeProperty( struct FPropertyChangedEvent& PropertyChangedEvent)
 {
 	Super::PostEditChangeProperty(PropertyChangedEvent);
@@ -161,11 +99,6 @@ void UWorldTileDetails::PostEditChangeProperty( struct FPropertyChangedEvent& Pr
 	{
 		ParentPackageNameChangedEvent.Broadcast();
 	}
-	else if (PropertyName == GET_MEMBER_NAME_CHECKED(UWorldTileDetails, StreamingLevels)
-			|| MemberPropertyName == GET_MEMBER_NAME_CHECKED(UWorldTileDetails, StreamingLevels))
-	{
-		StreamingLevelsChangedEvent.Broadcast();
-	}
 	else if (PropertyName == GET_MEMBER_NAME_CHECKED(UWorldTileDetails, NumLOD))
 	{
 		LODSettingsChangedEvent.Broadcast();
diff --git a/Engine/Source/Editor/WorldBrowser/Private/Tiles/WorldTileDetails.h b/Engine/Source/Editor/WorldBrowser/Private/Tiles/WorldTileDetails.h
--- a/Engine/Source/Editor/WorldBrowser/Private/Tiles/WorldTileDetails.h
+++ b/Engine/Source/Editor/WorldBrowser/Private/Tiles/WorldTileDetails.h
@@ -4,43 +4,6 @@
 
 class FWorldTileModel;
 
-/////////////////////////////////////////////////////
-// EStreamingLevelMode
-UENUM()
-namespace EStreamingLevelMode
-{
-	enum Type
-	{
-		AlwaysLoaded,
-		Blueprint,
-		Bounds,
-	};
-}
-
-/////////////////////////////////////////////////////
-// FTileStreamingLevelDetails
-//
-// Helper class to hold tile streaming levels descriptions
-//
-USTRUCT()
-struct FTileStreamingLevelDetails
-{
-	GENERATED_USTRUCT_BODY()
-	
-	// Method(trigger) for streaming this level
-	UPROPERTY(Category=StreamingLevel, EditAnywhere)	
-	TEnumAsByte<EStreamingLevelMode::Type>	StreamingMode;
-	
-	// Level long package name to stream-in
-	UPROPERTY(Category=StreamingLevel, EditAnywhere)
-	FName									PackageName;
-
-	// Helper functions to map streaming method to a streaming class and vice versa
-	static UClass* StreamingMode2Class(EStreamingLevelMode::Type InMode);
-	static EStreamingLevelMode::Type StreamingObject2Mode(ULevelStreaming* InLevelStreaming);
-};
-
-
 /////////////////////////////////////////////////////
 // FTileLODEntryDetails
 //
@@ -104,10 +67,6 @@ class UWorldTileDetails : public UObject
 	// Tile sorting order
 	UPROPERTY(Category=Tile, EditAnywhere, meta=(ClampMin = "-1000", ClampMax = "1000", UIMin = "-1000", UIMax = "1000"))
 	int32							ZOrder;
-
-	// List of streaming levels for this tile 
-	UPROPERTY(Category=Tile, EditAnywhere)
-	TArray<FTileStreamingLevelDetails>	StreamingLevels;
 	
 	// LOD entries number
 	UPROPERTY(Category=LODSettings, EditAnywhere, meta=(ClampMin = "0", ClampMax = "4", UIMin = "0", UIMax = "4"))
@@ -139,7 +98,6 @@ class UWorldTileDetails : public UObject
 
 	FTilePropertyChanged			PositionChangedEvent;
 	FTilePropertyChanged			ParentPackageNameChangedEvent;
-	FTilePropertyChanged			StreamingLevelsChangedEvent;
 	FTilePropertyChanged			LODSettingsChangedEvent;
 	FTilePropertyChanged			ZOrderChangedEvent;
 	
@@ -153,7 +111,5 @@ public:
 	
 	// Gets the initialized FWorldTileInfo from this details values
 	FWorldTileInfo GetInfo() const;
-	
-	// Initialize streaming levels descriptions with values stored in TileModel
-	void SyncStreamingLevels(const FWorldTileModel& TileModel);
+
 };
diff --git a/Engine/Source/Editor/WorldBrowser/Private/Tiles/WorldTileDetailsCustomization.cpp b/Engine/Source/Editor/WorldBrowser/Private/Tiles/WorldTileDetailsCustomization.cpp
--- a/Engine/Source/Editor/WorldBrowser/Private/Tiles/WorldTileDetailsCustomization.cpp
+++ b/Engine/Source/Editor/WorldBrowser/Private/Tiles/WorldTileDetailsCustomization.cpp
@@ -104,10 +104,6 @@ void FWorldTileDetailsCustomization::CustomizeDetails(IDetailLayoutBuilder& Deta
 		TileCategory.AddProperty(GET_MEMBER_NAME_CHECKED(UWorldTileDetails, ZOrder))
 			.IsEnabled(IsPropertyEnabled);
 		
-		// Streaming levels
-		TileCategory.AddProperty(GET_MEMBER_NAME_CHECKED(UWorldTileDetails, StreamingLevels))
-			.IsEnabled(IsPropertyEnabled);
-
 		// bTileEditable (invisible property to control other properties editable state)
 		TileEditableHandle = DetailLayoutBuilder.GetProperty(GET_MEMBER_NAME_CHECKED(UWorldTileDetails, bTileEditable));
 		TileCategory.AddProperty(TileEditableHandle)
diff --git a/Engine/Source/Editor/WorldBrowser/Private/Tiles/WorldTileModel.cpp b/Engine/Source/Editor/WorldBrowser/Private/Tiles/WorldTileModel.cpp
--- a/Engine/Source/Editor/WorldBrowser/Private/Tiles/WorldTileModel.cpp
+++ b/Engine/Source/Editor/WorldBrowser/Private/Tiles/WorldTileModel.cpp
@@ -34,7 +34,6 @@ FWorldTileModel::FWorldTileModel(const TWeakObjectPtr<UEditorEngine>& InEditor,
 	// Subscribe to tile properties changes
 	TileDetails->PositionChangedEvent.AddRaw(this, &FWorldTileModel::OnPositionPropertyChanged);
 	TileDetails->ParentPackageNameChangedEvent.AddRaw(this, &FWorldTileModel::OnParentPackageNamePropertyChanged);
-	TileDetails->StreamingLevelsChangedEvent.AddRaw(this, &FWorldTileModel::OnStreamingLevelsPropertyChanged);
 	TileDetails->LODSettingsChangedEvent.AddRaw(this, &FWorldTileModel::OnLODSettingsPropertyChanged);
 	TileDetails->ZOrderChangedEvent.AddRaw(this, &FWorldTileModel::OnZOrderPropertyChanged);
 			
@@ -44,7 +43,6 @@ FWorldTileModel::FWorldTileModel(const TWeakObjectPtr<UEditorEngine>& InEditor,
 		FWorldCompositionTile& Tile = WorldComposition->GetTilesList()[TileIdx];
 		
 		TileDetails->SetInfo(Tile.Info);
-		TileDetails->SyncStreamingLevels(*this);
 		TileDetails->PackageName = Tile.PackageName;
 		TileDetails->bPersistentLevel = false;
 				
@@ -85,7 +83,6 @@ FWorldTileModel::~FWorldTileModel()
 	{
 		TileDetails->PositionChangedEvent.RemoveAll(this);
 		TileDetails->ParentPackageNameChangedEvent.RemoveAll(this);
-		TileDetails->StreamingLevelsChangedEvent.RemoveAll(this);
 		
 		TileDetails->RemoveFromRoot();
 		TileDetails->MarkPendingKill();
@@ -356,33 +353,6 @@ ALandscapeProxy* FWorldTileModel::GetLandscape() const
 	return Landscape.Get();
 }
 
-void FWorldTileModel::AddStreamingLevel(UClass* InStreamingClass, const FName& InPackageName)
-{
-	if (LevelCollectionModel.IsReadOnly() || !IsEditable() || IsRootTile())
-	{
-		return;
-	}
-
-	UWorld* LevelWorld = CastChecked<UWorld>(GetLevelObject()->GetOuter());
-	// check uniqueness 
-	if (LevelWorld->StreamingLevels.FindMatch(ULevelStreaming::FPackageNameMatcher(InPackageName)) != INDEX_NONE)
-	{
-		return;
-	}
-		
-	ULevelStreaming* StreamingLevel = static_cast<ULevelStreaming*>(StaticConstructObject(InStreamingClass, LevelWorld, NAME_None, RF_NoFlags, NULL));
-	// Associate a package name.
-	StreamingLevel->PackageName			= InPackageName;
-	// Seed the level's draw color.
-	StreamingLevel->DrawColor			= FColor::MakeRandomColor();
-	StreamingLevel->LevelTransform		= FTransform::Identity;
-	StreamingLevel->PackageNameToLoad	= InPackageName;
-
-	// Add the streaming level to level's world
-	LevelWorld->StreamingLevels.Add(StreamingLevel);
-	LevelWorld->GetOutermost()->MarkPackageDirty();
-}
-
 void FWorldTileModel::AssignToLayer(const FWorldTileLayer& InLayer)
 {
 	if (LevelCollectionModel.IsReadOnly())
@@ -523,8 +493,7 @@ void FWorldTileModel::Update()
 		// Receive tile info from world composition
 		FWorldTileInfo Info = LevelCollectionModel.GetWorld()->WorldComposition->GetTileInfo(TileDetails->PackageName);
 		TileDetails->SetInfo(Info);
-		TileDetails->SyncStreamingLevels(*this);
-	
+			
 		ULevel* Level = GetLevelObject();
 		if (Level != nullptr && Level->bIsVisible)
 		{
@@ -684,26 +653,6 @@ void FWorldTileModel::EnsureLevelHasBoundsActor()
 	}
 }
 
-void FWorldTileModel::FixupStreamingObjects()
-{
-	check(LoadedLevel != NULL);
-
-	// Remove streaming levels that is not part of our world
-	UWorld* LevelWorld = CastChecked<UWorld>(LoadedLevel->GetOuter());
-	for(int32 LevelIndex = LevelWorld->StreamingLevels.Num() - 1; LevelIndex >= 0; --LevelIndex)
-	{
-		FName StreamingPackageName = LevelWorld->StreamingLevels[LevelIndex]->PackageName;
-		TSharedPtr<FLevelModel> LevelModel = LevelCollectionModel.FindLevelModel(StreamingPackageName);
-		if (!LevelModel.IsValid())
-		{
-			LevelWorld->StreamingLevels.RemoveAt(LevelIndex);
-			FFormatNamedArguments Arguments;
-			Arguments.Add(TEXT("LevelName"), FText::FromName(StreamingPackageName));
-			FMessageLog("MapCheck").Warning(FText::Format( LOCTEXT("MapCheck_Message_InvalidStreamingLevel", "Streaming level '{LevelName}' is not under the current world root. Fixed by removing from streaming list."), Arguments ) );
-		}
-	}
-}
-
 void FWorldTileModel::SortRecursive()
 {
 	AllChildren.Sort(FCompareByLongPackageName());
@@ -766,31 +715,6 @@ void FWorldTileModel::OnParentPackageNamePropertyChanged()
 	TileDetails->ParentPackageName = FName(*Info.ParentTilePackageName);
 }
 
-void FWorldTileModel::OnStreamingLevelsPropertyChanged()
-{
-	ULevel* Level = GetLevelObject();
-	if (Level == NULL)
-	{
-		TileDetails->StreamingLevels.Empty();
-		return;
-	}
-	
-	// Delete all streaming levels from the world objects
-	CastChecked<UWorld>(Level->GetOuter())->StreamingLevels.Empty();
-
-	// Recreate streaming levels using settings stored in the tile details
-	for (auto It = TileDetails->StreamingLevels.CreateIterator(); It; ++It)
-	{
-		FTileStreamingLevelDetails& StreamingLevelDetails = (*It);
-		FName PackageName = StreamingLevelDetails.PackageName;
-		if (PackageName != NAME_None && FPackageName::DoesPackageExist(PackageName.ToString()))
-		{
-			UClass* StreamingClass = FTileStreamingLevelDetails::StreamingMode2Class(StreamingLevelDetails.StreamingMode);
-			AddStreamingLevel(StreamingClass, PackageName);
-		}
-	}
-}
-
 void FWorldTileModel::OnLODSettingsPropertyChanged()
 {
 	OnLevelInfoUpdated();
diff --git a/Engine/Source/Editor/WorldBrowser/Private/Tiles/WorldTileModel.h b/Engine/Source/Editor/WorldBrowser/Private/Tiles/WorldTileModel.h
--- a/Engine/Source/Editor/WorldBrowser/Private/Tiles/WorldTileModel.h
+++ b/Engine/Source/Editor/WorldBrowser/Private/Tiles/WorldTileModel.h
@@ -178,10 +178,7 @@ private:
 	
 	/** Handler for ParentPackageName event from Tile details object  */
 	void OnParentPackageNamePropertyChanged();
-	
-	/** Handler for StreamingLevels event from Tile details object  */
-	void OnStreamingLevelsPropertyChanged();
-	
+
 	/** Handler for LOD settings changes event from Tile details object  */
 	void OnLODSettingsPropertyChanged();
 	
