World Browser: Fixed selection sync between all 3 panels. Adjusted toolbar buttons positions

diff --git a/Engine/Source/Editor/WorldBrowser/Private/SWorldDetails.cpp b/Engine/Source/Editor/WorldBrowser/Private/SWorldDetails.cpp
--- a/Engine/Source/Editor/WorldBrowser/Private/SWorldDetails.cpp
+++ b/Engine/Source/Editor/WorldBrowser/Private/SWorldDetails.cpp
@@ -82,7 +82,7 @@ void SWorldDetails::OnBrowseWorld(UWorld* InWorld)
 					]
 					
 					+SHorizontalBox::Slot()
-					.FillWidth(1.f)
+					.AutoWidth()
 					.VAlign(VAlign_Center)
 					.HAlign(HAlign_Left)
 					.Padding(4,0,4,0)
@@ -102,7 +102,7 @@ void SWorldDetails::OnBrowseWorld(UWorld* InWorld)
 					+SHorizontalBox::Slot()
 					.AutoWidth()
 					.VAlign(VAlign_Center)
-					.HAlign(HAlign_Right)
+					.HAlign(HAlign_Left)
 					[
 						SNew(SButton)
 						.ButtonStyle(FEditorStyle::Get(), "ToggleButton")
@@ -121,7 +121,7 @@ void SWorldDetails::OnBrowseWorld(UWorld* InWorld)
 					+SHorizontalBox::Slot()
 					.AutoWidth()
 					.VAlign(VAlign_Center)
-					.HAlign(HAlign_Right)
+					.HAlign(HAlign_Left)
 					[
 						SNew(SButton)
 						.Visibility(this, &SWorldDetails::GetCompositionButtonVisibility)
@@ -153,6 +153,9 @@ void SWorldDetails::OnBrowseWorld(UWorld* InWorld)
 		];
 
 		WorldModel->RegisterDetailsCustomization(PropertyModule, DetailsView);
+
+		// Current selection
+		OnSelectionChanged();
 	}
 }
 
diff --git a/Engine/Source/Editor/WorldBrowser/Private/SWorldHierarchy.cpp b/Engine/Source/Editor/WorldBrowser/Private/SWorldHierarchy.cpp
--- a/Engine/Source/Editor/WorldBrowser/Private/SWorldHierarchy.cpp
+++ b/Engine/Source/Editor/WorldBrowser/Private/SWorldHierarchy.cpp
@@ -193,7 +193,7 @@ public:
 		];
 	
 		WorldModel->AddFilter(SearchBoxLevelFilter.ToSharedRef());
-		RefreshView();
+		OnUpdateSelection();
 	}
 
 	/** Regenerates current items */
@@ -544,7 +544,7 @@ void SWorldHierarchy::OnBrowseWorld(UWorld* InWorld)
 					
 					// Toolbar
 					+SHorizontalBox::Slot()
-					.FillWidth(1.f)
+					.AutoWidth()
 					.VAlign(VAlign_Center)
 					.HAlign(HAlign_Left)
 					[
@@ -586,7 +586,7 @@ void SWorldHierarchy::OnBrowseWorld(UWorld* InWorld)
 					+SHorizontalBox::Slot()
 					.AutoWidth()
 					.VAlign(VAlign_Center)
-					.HAlign(HAlign_Right)
+					.HAlign(HAlign_Left)
 					[
 						SNew(SButton)
 						.ButtonStyle(FEditorStyle::Get(), "ToggleButton")
@@ -605,7 +605,7 @@ void SWorldHierarchy::OnBrowseWorld(UWorld* InWorld)
 					+SHorizontalBox::Slot()
 					.AutoWidth()
 					.VAlign(VAlign_Center)
-					.HAlign(HAlign_Right)
+					.HAlign(HAlign_Left)
 					[
 						SNew(SButton)
 						.Visibility(this, &SWorldHierarchy::GetCompositionButtonVisibility)
