World Browser: Generated assets for LOD maps now stored inside LOD map package
World Browser: LOD maps generator now correctly handles maps under SC
ProxyMesh utilitiy now produces assets in a separate packages or in one single package
ProxyMesh and MergeActors utilities now use more consistent names for a produced assets (demolet naming conventions)

diff --git a/Engine/Source/Developer/MeshUtilities/Private/MeshUtilities.cpp b/Engine/Source/Developer/MeshUtilities/Private/MeshUtilities.cpp
--- a/Engine/Source/Developer/MeshUtilities/Private/MeshUtilities.cpp
+++ b/Engine/Source/Developer/MeshUtilities/Private/MeshUtilities.cpp
@@ -531,8 +531,9 @@ private:
 	
 	virtual void CreateProxyMesh( 
 		const TArray<AActor*>& Actors, 
-		const struct FMeshProxySettings& InProxySettings, 
-		const FString& ProxyPackageName,
+		const struct FMeshProxySettings& InProxySettings,
+		UPackage* InOuter,
+		const FString& ProxyBasePackageName,
 		TArray<UObject*>& OutAssetsToSync,
 		FVector& OutProxyLocation
 		) OVERRIDE;
@@ -2984,7 +2985,8 @@ bool FMeshUtilities::BuildSkeletalMesh( FStaticLODModel& LODModel, const FRefere
 void FMeshUtilities::CreateProxyMesh( 
 	const TArray<AActor*>& SourceActors, 
 	const struct FMeshProxySettings& InProxySettings,
-	const FString& ProxyPackageName,
+	UPackage* InOuter,
+	const FString& ProxyBasePackageName,
 	TArray<UObject*>& OutAssetsToSync,
 	FVector& OutProxyLocation)
 {
@@ -3009,6 +3011,9 @@ void FMeshUtilities::CreateProxyMesh(
 		{
 			TArray<UStaticMeshComponent*> Components;
 			Actor->GetComponents<UStaticMeshComponent>(Components);
+			// TODO: support instanced static meshes
+			Components.RemoveAll([](UStaticMeshComponent* Val){ return Val->IsA(UInstancedStaticMeshComponent::StaticClass()); });
+			//
 			ComponentsToMerge.Append(Components);
 		}
 	}
@@ -3113,32 +3118,25 @@ void FMeshUtilities::CreateProxyMesh(
 	}
 	
 	//
-	//Create New Package for the Proxy mesh
+	// Base asset name for a new assets
 	//
-	FString PackageName = ProxyPackageName;
-	FString AssetName = TEXT("ProxyMesh");
-	if (PackageName.IsEmpty())
-	{
-		// Proxy package will be stored in Content folder
-		PackageName = FPaths::GameContentDir() + AssetName;
-	}
-	else
-	{
-		AssetName = FPackageName::GetShortName(PackageName);
-	}
-
-	UPackage* Package = CreatePackage(NULL, *PackageName);
-
-	check(Package);
-	Package->FullyLoad();
-	Package->Modify();
+	const FString AssetBaseName = FPackageName::GetShortName(ProxyBasePackageName);
+	const FString AssetBasePath = FPackageName::IsShortPackageName(ProxyBasePackageName) ? 
+		FPackageName::FilenameToLongPackageName(FPaths::GameContentDir()) : (FPackageName::GetLongPackagePath(ProxyBasePackageName) + TEXT("/"));
 
 	// Construct proxy material
-	UMaterial* ProxyMaterial = MaterialExportUtils::CreateMaterial(ProxyFlattenMaterial, Package, AssetName, RF_Public|RF_Standalone);
+	UMaterial* ProxyMaterial = MaterialExportUtils::CreateMaterial(ProxyFlattenMaterial, InOuter, ProxyBasePackageName, RF_Public|RF_Standalone);
 	
 	// Construct proxy static mesh
-	FName StaticMeshName = MakeUniqueObjectName(Package, UStaticMesh::StaticClass(), *AssetName);
-	UStaticMesh* StaticMesh = new(Package, StaticMeshName, RF_Public|RF_Standalone) UStaticMesh(FPostConstructInitializeProperties());
+	UPackage* MeshPackage = InOuter;
+	if (MeshPackage == nullptr)
+	{
+		MeshPackage = CreatePackage(NULL, *(AssetBasePath + TEXT("SM_") + AssetBaseName));
+		MeshPackage->FullyLoad();
+		MeshPackage->Modify();
+	}
+
+	UStaticMesh* StaticMesh = new(MeshPackage, FName(*(TEXT("SM_") + AssetBaseName)), RF_Public|RF_Standalone) UStaticMesh(FPostConstructInitializeProperties());
 	StaticMesh->InitResources();
 	{
 		FString OutputPath = StaticMesh->GetPathName();
@@ -3173,7 +3171,12 @@ void FMeshUtilities::CreateProxyMesh(
 	{
 		if (FlatMat.DiffuseSamples.Num() > 1)
 		{
-			FString DiffuseTextureName = MakeUniqueObjectName(Package, UTexture2D::StaticClass(), TEXT("FlattenMaterial_Diffuse")).ToString();
+			FString DiffuseTextureName = MakeUniqueObjectName(Package, UTexture2D::StaticClass(), *(TEXT("T_FLATTEN_") + AssetBaseName + TEXT("_D"))).ToString();
+			
+			UPackage* TexPackage = CreatePackage(NULL, *(AssetBasePath + DiffuseTextureName));
+			TexPackage->FullyLoad();
+			TexPackage->Modify();
+					
 			FCreateTexture2DParameters TexParams;
 			TexParams.bUseAlpha = false;
 			TexParams.CompressionSettings = TC_Default;
@@ -3184,7 +3187,7 @@ void FMeshUtilities::CreateProxyMesh(
 				FlatMat.DiffuseSize.X, 
 				FlatMat.DiffuseSize.Y,
 				FlatMat.DiffuseSamples,
-				Package,
+				TexPackage,
 				DiffuseTextureName,
 				RF_Public|RF_Standalone, 
 				TexParams);
@@ -3582,24 +3585,21 @@ void FMeshUtilities::MergeActors(
 		FString PackageName;
 		if (InPackageName.IsEmpty())
 		{
-			AssetName = FString(TEXT("MergedMesh_")) + FPackageName::GetShortName(MergedMesh.AssetPackageName);
-			PackageName = FPackageName::GetLongPackagePath(MergedMesh.AssetPackageName) + AssetName;
+			AssetName = TEXT("SM_MERGED_") + FPackageName::GetShortName(MergedMesh.AssetPackageName);
+			PackageName = FPackageName::GetLongPackagePath(MergedMesh.AssetPackageName) + TEXT("/") + AssetName;
 		}
 		else
 		{
 			AssetName = FPackageName::GetShortName(InPackageName);
 			PackageName = InPackageName;
 		}
-		
-		FName UniquePackageName = MakeUniqueObjectName(NULL, UPackage::StaticClass(), *PackageName);
 
-		UPackage* Package = CreatePackage(NULL, *UniquePackageName.ToString());
+		UPackage* Package = CreatePackage(NULL, *PackageName);
 		check(Package);
 		Package->FullyLoad();
 		Package->Modify();
 
-		FName StaticMeshName = MakeUniqueObjectName(Package, UStaticMesh::StaticClass(), *AssetName);
-		UStaticMesh* StaticMesh = new(Package, StaticMeshName, RF_Public|RF_Standalone) UStaticMesh(FPostConstructInitializeProperties());
+		UStaticMesh* StaticMesh = new(Package, *AssetName, RF_Public|RF_Standalone) UStaticMesh(FPostConstructInitializeProperties());
 		StaticMesh->InitResources();
 		
 		FString OutputPath = StaticMesh->GetPathName();
diff --git a/Engine/Source/Developer/MeshUtilities/Public/MeshUtilities.h b/Engine/Source/Developer/MeshUtilities/Public/MeshUtilities.h
--- a/Engine/Source/Developer/MeshUtilities/Public/MeshUtilities.h
+++ b/Engine/Source/Developer/MeshUtilities/Public/MeshUtilities.h
@@ -249,14 +249,16 @@ public:
 	 *
 	 *	@param	Actors					List of Actors to merge
 	 *	@param	InProxySettings			Merge settings
-	 *	@param	ProxyPackageName		Destination package name for a generated assets
+	 *	@param	InOuter					Package for a generated assets, if NULL new packages will be created for each asset
+	 *	@param	ProxyBasePackageName	Will be used for naming generated assets, in case InOuter is not specified ProxyBasePackageName will be used as long package name for creating new packages
 	 *	@param	OutAssetsToSync			Result assets - mesh, material
 	 *	@param	OutProxyLocation		Proxy mesh location in the world (bounding box origin of merged actors)
 	 */
 	virtual void CreateProxyMesh(
 		const TArray<AActor*>& Actors,
 		const struct FMeshProxySettings& InProxySettings,
-		const FString& ProxyPackageName,
+		UPackage* InOuter,
+		const FString& ProxyBasePackageName,
 		TArray<UObject*>& OutAssetsToSync,
 		FVector& OutProxyLocation
 		) = 0;
diff --git a/Engine/Source/Editor/LevelEditor/Private/LevelEditorActions.cpp b/Engine/Source/Editor/LevelEditor/Private/LevelEditorActions.cpp
--- a/Engine/Source/Editor/LevelEditor/Private/LevelEditorActions.cpp
+++ b/Engine/Source/Editor/LevelEditor/Private/LevelEditorActions.cpp
@@ -2782,7 +2782,7 @@ void FLevelEditorCommands::RegisterCommands()
 	UI_COMMAND( LockGroup, "Lock", "Locks the selected groups", EUserInterfaceActionType::Button, FInputGesture() );
 	UI_COMMAND( UnlockGroup, "Unlock", "Unlocks the selected groups", EUserInterfaceActionType::Button, FInputGesture() );
 	UI_COMMAND( MergeActors, "Create Mesh Proxy...", "Harvest geometry from selected actors and merge them into single mesh", EUserInterfaceActionType::Button, FInputGesture() );
-	UI_COMMAND( MergeActorsByMaterials, "Merge Actors", "Harvest geometry from selected actors and merge grouping them by materials", EUserInterfaceActionType::Button, FInputGesture() );
+	UI_COMMAND( MergeActorsByMaterials, "Merge Actors...", "Harvest geometry from selected actors and merge grouping them by materials", EUserInterfaceActionType::Button, FInputGesture() );
 
 	UI_COMMAND( ShowAll, "Show All Actors", "Shows all actors", EUserInterfaceActionType::Button, FInputGesture( EModifierKey::Control, EKeys::H ) );
 	UI_COMMAND( ShowSelectedOnly, "Show Only Selected", "Shows only the selected actors", EUserInterfaceActionType::Button, FInputGesture() );
diff --git a/Engine/Source/Editor/UnrealEd/Private/Dialogs/SMeshMergingDialog.cpp b/Engine/Source/Editor/UnrealEd/Private/Dialogs/SMeshMergingDialog.cpp
--- a/Engine/Source/Editor/UnrealEd/Private/Dialogs/SMeshMergingDialog.cpp
+++ b/Engine/Source/Editor/UnrealEd/Private/Dialogs/SMeshMergingDialog.cpp
@@ -287,7 +287,7 @@ void SMeshMergingDialog::GenerateNewPackageName()
 				if (Component->StaticMesh)
 				{
 					MergedMeshPackageName = FPackageName::GetLongPackagePath(Component->StaticMesh->GetOutermost()->GetName());
-					MergedMeshPackageName+= FString(TEXT("/MergedMesh_")) + Component->StaticMesh->GetName();
+					MergedMeshPackageName+= FString(TEXT("/SM_MERGED_")) + Component->StaticMesh->GetName();
 					break;
 				}
 			}
@@ -301,10 +301,9 @@ void SMeshMergingDialog::GenerateNewPackageName()
 
 	if (MergedMeshPackageName.IsEmpty())
 	{
-		MergedMeshPackageName = FPackageName::FilenameToLongPackageName(FPaths::GameContentDir() + TEXT("MergedMesh"));
+		MergedMeshPackageName = FPackageName::FilenameToLongPackageName(FPaths::GameContentDir() + TEXT("SM_MERGED"));
+		MergedMeshPackageName = MakeUniqueObjectName(NULL, UPackage::StaticClass(), *MergedMeshPackageName).ToString();
 	}
-
-	MergedMeshPackageName = MakeUniqueObjectName(NULL, UPackage::StaticClass(), *MergedMeshPackageName).ToString();
 }
 
 FText SMeshMergingDialog::GetMergedMeshPackageName() const
diff --git a/Engine/Source/Editor/UnrealEd/Private/Dialogs/SMeshProxyDialog.cpp b/Engine/Source/Editor/UnrealEd/Private/Dialogs/SMeshProxyDialog.cpp
--- a/Engine/Source/Editor/UnrealEd/Private/Dialogs/SMeshProxyDialog.cpp
+++ b/Engine/Source/Editor/UnrealEd/Private/Dialogs/SMeshProxyDialog.cpp
@@ -647,7 +647,7 @@ FReply SMeshProxyDialog::OnMergeClicked()
 		GEditor->BeginTransaction( LOCTEXT( "MeshProxy_Create", "Creating Mesh Proxy" ) );
 
 		FVector ProxyLocation = FVector::ZeroVector;
-		MeshUtilities.CreateProxyMesh(Actors, ProxySettings, ProxyPackageName, AssetsToSync, ProxyLocation);
+		MeshUtilities.CreateProxyMesh(Actors, ProxySettings, NULL, ProxyPackageName, AssetsToSync, ProxyLocation);
 
 		GEditor->EndTransaction();
 		GWarn->EndSlowTask();
@@ -710,7 +710,7 @@ void SMeshProxyDialog::GenerateNewProxyPackageName()
 				if (Component->StaticMesh)
 				{
 					ProxyPackageName = FPackageName::GetLongPackagePath(Component->StaticMesh->GetOutermost()->GetName());
-					ProxyPackageName+= FString(TEXT("/ProxyMesh_")) + Component->StaticMesh->GetName();
+					ProxyPackageName+= FString(TEXT("/PROXY_")) + Component->StaticMesh->GetName();
 					break;
 				}
 			}
@@ -724,10 +724,9 @@ void SMeshProxyDialog::GenerateNewProxyPackageName()
 
 	if (ProxyPackageName.IsEmpty())
 	{
-		ProxyPackageName = FPackageName::FilenameToLongPackageName(FPaths::GameContentDir() + TEXT("ProxyMesh"));
+		ProxyPackageName = FPackageName::FilenameToLongPackageName(FPaths::GameContentDir() + TEXT("PROXY"));
+		ProxyPackageName = MakeUniqueObjectName(NULL, UPackage::StaticClass(), *ProxyPackageName).ToString();
 	}
-
-	ProxyPackageName = MakeUniqueObjectName(NULL, UPackage::StaticClass(), *ProxyPackageName).ToString();
 }
 
 FText SMeshProxyDialog::GetProxyPackagName() const
diff --git a/Engine/Source/Editor/UnrealEd/Private/EditorExporters.cpp b/Engine/Source/Editor/UnrealEd/Private/EditorExporters.cpp
--- a/Engine/Source/Editor/UnrealEd/Private/EditorExporters.cpp
+++ b/Engine/Source/Editor/UnrealEd/Private/EditorExporters.cpp
@@ -2843,11 +2843,23 @@ namespace MaterialExportUtils
 		return true;
 	}
 
-	UMaterial* CreateMaterial(const FFlattenMaterial& InFlattenMaterial, UObject* Outer, const FString& BaseName, EObjectFlags Flags)
+	UMaterial* CreateMaterial(const FFlattenMaterial& InFlattenMaterial, UPackage* Outer, const FString& BaseName, EObjectFlags Flags)
 	{
-		FName MaterialName = MakeUniqueObjectName(Outer, UMaterial::StaticClass(), *(BaseName + TEXT("_Material")));
-		
-		UMaterial* Material = ConstructObject<UMaterial>(UMaterial::StaticClass(), Outer, MaterialName, Flags);
+		const FString AssetBaseName = FPackageName::GetShortName(BaseName);
+		const FString AssetBasePath = FPackageName::IsShortPackageName(BaseName) ? 
+			FPackageName::FilenameToLongPackageName(FPaths::GameContentDir()) : (FPackageName::GetLongPackagePath(BaseName) + TEXT("/"));
+				
+		// Create material
+		const FString MaterialAssetName = TEXT("M_") + AssetBaseName;
+		UPackage* MaterialOuter = Outer;
+		if (MaterialOuter == NULL)
+		{
+			MaterialOuter = CreatePackage(NULL, *(AssetBasePath + MaterialAssetName));
+			MaterialOuter->FullyLoad();
+			MaterialOuter->Modify();
+		}
+
+		UMaterial* Material = ConstructObject<UMaterial>(UMaterial::StaticClass(), MaterialOuter, FName(*MaterialAssetName), Flags);
 		Material->TwoSided = false;
 		Material->SetShadingModel(MSM_DefaultLit);
 
@@ -2876,19 +2888,27 @@ namespace MaterialExportUtils
 		//Build the Diffuse UTexture
 		if (InFlattenMaterial.DiffuseSamples.Num() > 1)
 		{
-			FString DiffuseTextureName = MakeUniqueObjectName(Outer, UTexture2D::StaticClass(), *(BaseName + TEXT("_Diffuse"))).ToString();
 			FCreateTexture2DParameters TexParams;
 			TexParams.bUseAlpha = false;
 			TexParams.CompressionSettings = TC_Default;
 			TexParams.bDeferCompression = true;
 			TexParams.bSRGB = false;
 
+			const FString TextureAssetName = TEXT("T_") + AssetBaseName + TEXT("_D");
+			UPackage* TextureOuter = Outer;
+			if (TextureOuter == NULL)
+			{
+				TextureOuter = CreatePackage(NULL, *(AssetBasePath + TextureAssetName));
+				TextureOuter->FullyLoad();
+				TextureOuter->Modify();
+			}
+
 			UTexture2D* DiffuseTexture = FImageUtils::CreateTexture2D(
 				InFlattenMaterial.DiffuseSize.X, 
 				InFlattenMaterial.DiffuseSize.Y,
 				InFlattenMaterial.DiffuseSamples,
-				Outer,
-				DiffuseTextureName,
+				TextureOuter,
+				TextureAssetName,
 				Flags, 
 				TexParams);
 
@@ -2905,19 +2925,27 @@ namespace MaterialExportUtils
 		//Build the Normal UTexture
 		if (InFlattenMaterial.NormalSamples.Num() > 1)
 		{
-			FString NormalTextureName = MakeUniqueObjectName(Outer, UTexture2D::StaticClass(), *(BaseName + TEXT("_Normal"))).ToString();
 			FCreateTexture2DParameters TexParams;
 			TexParams.bUseAlpha = false;
 			TexParams.CompressionSettings = TC_Normalmap;
 			TexParams.bDeferCompression = true;
 			TexParams.bSRGB = false;
 
+			const FString TextureAssetName = TEXT("T_") + AssetBaseName + TEXT("_N");
+			UPackage* TextureOuter = Outer;
+			if (TextureOuter == NULL)
+			{
+				TextureOuter = CreatePackage(NULL, *(AssetBasePath + TextureAssetName));
+				TextureOuter->FullyLoad();
+				TextureOuter->Modify();
+			}
+
 			UTexture2D* NormalTexture = FImageUtils::CreateTexture2D(
 				InFlattenMaterial.NormalSize.X, 
 				InFlattenMaterial.NormalSize.Y,
 				InFlattenMaterial.NormalSamples,
-				Outer,
-				NormalTextureName,
+				TextureOuter,
+				TextureAssetName,
 				Flags, 
 				TexParams);
 
diff --git a/Engine/Source/Editor/UnrealEd/Public/MaterialExportUtils.h b/Engine/Source/Editor/UnrealEd/Public/MaterialExportUtils.h
--- a/Engine/Source/Editor/UnrealEd/Public/MaterialExportUtils.h
+++ b/Engine/Source/Editor/UnrealEd/Public/MaterialExportUtils.h
@@ -55,11 +55,11 @@ namespace MaterialExportUtils
 	/**
 	 * Creates UMaterial object from a flatten material
 	 *
-	 * @param Outer			Outer for the material and texture objects.
-	 * @param BaseName		BaseName for the material and texture objects.
+	 * @param Outer			Outer for the material and texture objects, if NULL new packages will be created for each asset
+	 * @param BaseName		BaseName for the material and texture objects, should be a long package name in case Outer is not specified
 	 * @param Flags			Object flags for the material and texture objects.
 	 * @return				Returns a pointer to the constructed UMaterial object.
 	 */
-	UNREALED_API UMaterial* CreateMaterial(const FFlattenMaterial& InFlattenMaterial, UObject* Outer,const FString& BaseName, EObjectFlags Flags);
+	UNREALED_API UMaterial* CreateMaterial(const FFlattenMaterial& InFlattenMaterial, UPackage* InOuter,const FString& BaseName, EObjectFlags Flags);
 }
 
diff --git a/Engine/Source/Editor/WorldBrowser/Private/Tiles/WorldTileCollectionModel.cpp b/Engine/Source/Editor/WorldBrowser/Private/Tiles/WorldTileCollectionModel.cpp
--- a/Engine/Source/Editor/WorldBrowser/Private/Tiles/WorldTileCollectionModel.cpp
+++ b/Engine/Source/Editor/WorldBrowser/Private/Tiles/WorldTileCollectionModel.cpp
@@ -1713,264 +1713,247 @@ bool FWorldTileCollectionModel::HasGenerateLODLevelSupport() const
 	return bMeshProxyAvailable;
 }
 
-bool FWorldTileCollectionModel::GenerateLODLevel(TSharedPtr<FLevelModel> InLevelModel, int32 TargetLODIndex)
+bool FWorldTileCollectionModel::GenerateLODLevels(FLevelModelList InLevelList, int32 TargetLODIndex)
 {
 	if (!HasGenerateLODLevelSupport())
 	{
 		return false;
 	}
-	
-	ULevel* SourceLevel = InLevelModel->GetLevelObject();
-	if (SourceLevel == nullptr)
-	{
-		return false;
-	}
-
-	TSharedPtr<FWorldTileModel> TileModel = StaticCastSharedPtr<FWorldTileModel>(InLevelModel);
-	FWorldTileInfo TileInfo = TileModel->TileDetails->GetInfo();
-
-	if (!TileInfo.LODList.IsValidIndex(TargetLODIndex))
-	{
-		return false;
-	}
-
-	// We have to make original level visible, to correctly export it
-	const bool bVisibleLevel = InLevelModel->IsVisible();
-	if (!bVisibleLevel)
-	{
-		GetWorld()->WorldComposition->bTemporallyDisableOriginTracking = true;
-		InLevelModel->SetVisible(true);
-	}
 
 	IMeshUtilities& MeshUtilities = FModuleManager::Get().LoadModuleChecked<IMeshUtilities>("MeshUtilities");
 	IMeshMerging* MeshMerging = MeshUtilities.GetMeshMergingInterface();
 	check(MeshMerging);
 
-	GWarn->BeginSlowTask(LOCTEXT("GenerateLODLevel", "Generating LOD Level"), true);
-			
-	FWorldTileLODInfo LODInfo = TileModel->TileDetails->GetInfo().LODList[TargetLODIndex];
-			
-	// Target folder to save generated assets: /Content/LevelsLOD/PackageName/LODIndex/
-	const FString LODContentDir =
-		FPackageName::FilenameToLongPackageName(FPaths::GameContentDir()) +
-		FString::Printf(TEXT("MapsLOD/%s/%d/"), *FPackageName::GetShortName(TileModel->TileDetails->PackageName), TargetLODIndex + 1);
-
-	// Target filename for generated level: /LongPackageName+LOD/PackageName_LOD[LodIndex].umap
-	const FString LODLevelFilename =
-		FPackageName::LongPackageNameToFilename(
-			FString::Printf(
-				TEXT("%sLOD/%s_LOD%d"),	
-				*TileModel->TileDetails->PackageName.ToString(), 
-				*FPackageName::GetShortName(TileModel->TileDetails->PackageName), 
-				TargetLODIndex+1)) +
-		FPackageName::GetMapPackageExtension();
-		
-	// This is current actors offset from their original position
-	FVector ActorsOffset = FVector(TileModel->GetAbsoluteLevelPosition() - CurrentWorld->GlobalOriginOffset);
-	if (GetWorld()->WorldComposition->bTemporallyDisableOriginTracking)
+	// Select tiles that can be processed
+	TArray<TSharedPtr<FWorldTileModel>> TilesToProcess;
+	TArray<FString>						LODPackageNames;
+	for (TSharedPtr<FLevelModel> LevelModel : InLevelList)
 	{
-		ActorsOffset = FVector::ZeroVector;
+		ULevel* SourceLevel = LevelModel->GetLevelObject();
+		if (SourceLevel == nullptr)
+		{
+			continue;
+		}
+
+		TSharedPtr<FWorldTileModel> TileModel = StaticCastSharedPtr<FWorldTileModel>(LevelModel);
+		FWorldTileInfo TileInfo = TileModel->TileDetails->GetInfo();
+		if (!TileInfo.LODList.IsValidIndex(TargetLODIndex))
+		{
+			continue;
+		}
+
+		TilesToProcess.Add(TileModel);
 	}
+
+	// TODO: Need to check out all LOD maps here
 	
-	TArray<UObject*>			GeneratedAssets;
-	TArray<UStaticMesh*>		AssetsToSpawn;
-	TArray<FTransform>			AssetsToSpawnTransform;
-	TArray<AActor*>				Actors;
-	TArray<ALandscapeProxy*>	LandscapeActors;
-	// Separate flies from cutlets
-	for (AActor* Actor : SourceLevel->Actors)
-	{
-		if (Actor)
+	GWarn->BeginSlowTask(LOCTEXT("GenerateLODLevel", "Generating LOD Levels"), true);
+	// Generate LOD maps for each tile
+	for (TSharedPtr<FWorldTileModel> TileModel : TilesToProcess)
+	{
+		// We have to make original level visible, to correctly export it
+		const bool bVisibleLevel = TileModel->IsVisible();
+		if (!bVisibleLevel)
 		{
-			ALandscapeProxy* LandscapeProxy = Cast<ALandscapeProxy>(Actor);
-			if (LandscapeProxy)
-			{
-				LandscapeActors.Add(LandscapeProxy);
-			}
-			else
+			GetWorld()->WorldComposition->bTemporallyDisableOriginTracking = true;
+			TileModel->SetVisible(true);
+		}
+				
+		FWorldTileInfo TileInfo = TileModel->TileDetails->GetInfo();
+		FWorldTileLODInfo LODInfo = TileModel->TileDetails->GetInfo().LODList[TargetLODIndex];
+		
+		// Source level package name
+		FString SourceLongPackageName = TileModel->TileDetails->PackageName.ToString();
+		FString SourceShortPackageName = FPackageName::GetShortName(SourceLongPackageName);
+		// Target PackageName for generated level: /LongPackageName+LOD/ShortPackageName_LOD[LODIndex]
+		const FString LODLevelPackageName = FString::Printf(TEXT("%sLOD/%s_LOD%d"),	*SourceLongPackageName, *SourceShortPackageName, TargetLODIndex+1);
+		// Target level filename
+		const FString LODLevelFileName = FPackageName::LongPackageNameToFilename(LODLevelPackageName) + FPackageName::GetMapPackageExtension();
+
+		// Create a package for a LOD level
+		UPackage* LODPackage = CreatePackage(NULL, *LODLevelPackageName);
+		LODPackage->FullyLoad();
+		LODPackage->Modify();
+		// This is a hack to avoid save file dialog when we will be saving LOD map package
+		LODPackage->FileName = FName(*LODLevelFileName);
+				
+		// This is current actors offset from their original position
+		FVector ActorsOffset = FVector(TileModel->GetAbsoluteLevelPosition() - CurrentWorld->GlobalOriginOffset);
+		if (GetWorld()->WorldComposition->bTemporallyDisableOriginTracking)
+		{
+			ActorsOffset = FVector::ZeroVector;
+		}
+	
+		TArray<UStaticMesh*>		AssetsToSpawn;
+		TArray<FTransform>			AssetsToSpawnTransform;
+		TArray<AActor*>				Actors;
+		TArray<ALandscapeProxy*>	LandscapeActors;
+		// Separate flies from cutlets
+		for (AActor* Actor : TileModel->GetLevelObject()->Actors)
+		{
+			if (Actor)
 			{
-				Actors.Add(Actor);
+				ALandscapeProxy* LandscapeProxy = Cast<ALandscapeProxy>(Actor);
+				if (LandscapeProxy)
+				{
+					LandscapeActors.Add(LandscapeProxy);
+				}
+				else
+				{
+					Actors.Add(Actor);
+				}
 			}
 		}
-	}
 	
-	// Generate Proxy LOD mesh for all actors excluding landscapes
-	if (Actors.Num())
-	{
-		FMeshProxySettings ProxySettings;
-		ProxySettings.ScreenSize = ProxySettings.ScreenSize*(LODInfo.GenDetailsPercentage/100.f);
-		TArray<UObject*> OutAssets;
-		FVector OutProxyLocation;
-		FString ProxyPackageName = LODContentDir + TEXT("ProxyMesh");
-		ProxyPackageName = MakeUniqueObjectName(NULL, UPackage::StaticClass(), *ProxyPackageName).ToString();
-
-		MeshUtilities.CreateProxyMesh(Actors, ProxySettings, ProxyPackageName, OutAssets, OutProxyLocation);
-		
-		if (OutAssets.Num())
+		// Generate Proxy LOD mesh for all actors excluding landscapes
+		if (Actors.Num())
 		{
-			GeneratedAssets.Append(OutAssets);
-			UStaticMesh* ProxyMesh = nullptr;
-			if (OutAssets.FindItemByClass(&ProxyMesh))
+			GWarn->StatusUpdate(0, 10, LOCTEXT("GeneratingProxyMesh", "Generating Proxy Mesh"));
+
+			FMeshProxySettings ProxySettings;
+			ProxySettings.ScreenSize = ProxySettings.ScreenSize*(LODInfo.GenDetailsPercentage/100.f);
+			TArray<UObject*> OutAssets;
+			FVector OutProxyLocation;
+			FString ProxyPackageName = TEXT("PROXY_") + FPackageName::GetShortName(TileModel->TileDetails->PackageName);
+			// Generate proxy mesh and proxy material assets
+			MeshUtilities.CreateProxyMesh(Actors, ProxySettings, LODPackage, ProxyPackageName, OutAssets, OutProxyLocation);
+		
+			if (OutAssets.Num())
 			{
-				AssetsToSpawn.Add(ProxyMesh);
-				AssetsToSpawnTransform.Add(FTransform(OutProxyLocation - ActorsOffset));
+				UStaticMesh* ProxyMesh = nullptr;
+				if (OutAssets.FindItemByClass(&ProxyMesh))
+				{
+					AssetsToSpawn.Add(ProxyMesh);
+					AssetsToSpawnTransform.Add(FTransform(OutProxyLocation - ActorsOffset));
+				}
 			}
 		}
-	}
 
-	using namespace MaterialExportUtils;
+		using namespace MaterialExportUtils;
 
-	// Convert landscape actors into static meshes and apply mesh reduction 
-	int32 LandscapeActorIndex = 0;
-	for (ALandscapeProxy* Landscape : LandscapeActors)
-	{
-		GWarn->StatusUpdate(LandscapeActorIndex, LandscapeActors.Num(), LOCTEXT("ExportingLandscape", "Exporting Landscape Actors"));
-		
-		FRawMesh LandscapeRawMesh;
-		FFlattenMaterial LandscapeFlattenMaterial;
-		FVector LandscapeWorldLocation = Landscape->GetActorLocation();
+		// Convert landscape actors into static meshes and apply mesh reduction 
+		int32 LandscapeActorIndex = 0;
+		for (ALandscapeProxy* Landscape : LandscapeActors)
+		{
+			GWarn->StatusUpdate(LandscapeActorIndex, LandscapeActors.Num(), LOCTEXT("ExportingLandscape", "Exporting Landscape Actors"));
 		
-		Landscape->ExportToRawMesh(MAX_int32, LandscapeRawMesh);
+			FRawMesh LandscapeRawMesh;
+			FFlattenMaterial LandscapeFlattenMaterial;
+			FVector LandscapeWorldLocation = Landscape->GetActorLocation();
 		
-		for (FVector& VertexPos : LandscapeRawMesh.VertexPositions)
-		{
-			VertexPos-= LandscapeWorldLocation;
-		}
-					
-		// This is texture resolution for a landscape, probably need to be calculated using landscape size
-		LandscapeFlattenMaterial.DiffuseSize	= FIntPoint(1024, 1024);
-		LandscapeFlattenMaterial.NormalSize		= FIntPoint(1024, 1024);
-		ExportMaterial(Landscape, LandscapeFlattenMaterial);
+			Landscape->ExportToRawMesh(MAX_int32, LandscapeRawMesh);
 		
-		// TODO: Disabled landscape mesh simplification, does not really needed since exporting highest landscape LOD already provide low triangle count
-		// Reduce landscape mesh
-		//if (LODInfo.GenDetailsPercentage != 100.f)
-		//{
-		//	FMeshReductionSettings Settings;
-		//	Settings.PercentTriangles = LODInfo.GenDetailsPercentage/100.f;
-		//	float OutMaxDeviation = 0.f;
-		//	FRawMesh LandscapeReducedMesh;
-
-		//	MeshUtilities.GetMeshReductionInterface()->Reduce(
-		//		LandscapeReducedMesh,
-		//		OutMaxDeviation,
-		//		LandscapeRawMesh,
-		//		Settings);
-
-		//	LandscapeRawMesh = LandscapeReducedMesh;
-		//}
-
-		FString PackageName = LODContentDir + TEXT("LandscapeStatic");
-		PackageName = MakeUniqueObjectName(NULL, UPackage::StaticClass(), *PackageName).ToString();
+			for (FVector& VertexPos : LandscapeRawMesh.VertexPositions)
+			{
+				VertexPos-= LandscapeWorldLocation;
+			}
+								
+			// This is texture resolution for a landscape, probably need to be calculated using landscape size
+			LandscapeFlattenMaterial.DiffuseSize	= FIntPoint(1024, 1024);
+			LandscapeFlattenMaterial.NormalSize		= FIntPoint(1024, 1024);
+			ExportMaterial(Landscape, LandscapeFlattenMaterial);
 		
-		UPackage* Package = CreatePackage(NULL, *PackageName);
-		Package->FullyLoad();
-		Package->Modify();
-
-		// Construct landscape material
-		UMaterial* StaticLandscapeMaterial = MaterialExportUtils::CreateMaterial(
-			LandscapeFlattenMaterial, Package, TEXT("FlattenLandscapeMaterial"), RF_Public|RF_Standalone);
-		// Currently landscape exports world space normal map
-		StaticLandscapeMaterial->bTangentSpaceNormal = false;
-		StaticLandscapeMaterial->PostEditChange();
+			// TODO: Disabled landscape mesh simplification, does not really needed since exporting highest landscape LOD already provide low triangle count
+			// Reduce landscape mesh
+			//if (LODInfo.GenDetailsPercentage != 100.f)
+			//{
+			//	FMeshReductionSettings Settings;
+			//	Settings.PercentTriangles = LODInfo.GenDetailsPercentage/100.f;
+			//	float OutMaxDeviation = 0.f;
+			//	FRawMesh LandscapeReducedMesh;
+
+			//	MeshUtilities.GetMeshReductionInterface()->Reduce(
+			//		LandscapeReducedMesh,
+			//		OutMaxDeviation,
+			//		LandscapeRawMesh,
+			//		Settings);
+
+			//	LandscapeRawMesh = LandscapeReducedMesh;
+			//}
+
+			FString LandscapeBaseAssetName = Landscape->GetName();
+			// Construct landscape material
+			UMaterial* StaticLandscapeMaterial = MaterialExportUtils::CreateMaterial(
+				LandscapeFlattenMaterial, LODPackage, *LandscapeBaseAssetName, RF_Public|RF_Standalone);
+			// Currently landscape exports world space normal map
+			StaticLandscapeMaterial->bTangentSpaceNormal = false;
+			StaticLandscapeMaterial->PostEditChange();
 	
-		// Construct landscape static mesh
-		FName StaticMeshName = MakeUniqueObjectName(Package, UStaticMesh::StaticClass(), TEXT("LandscapeStaticMesh"));
-		UStaticMesh* StaticMesh = new(Package, StaticMeshName, RF_Public|RF_Standalone) UStaticMesh(FPostConstructInitializeProperties());
-		StaticMesh->InitResources();
-		{
-			FString OutputPath = StaticMesh->GetPathName();
-
-			// make sure it has a new lighting guid
-			StaticMesh->LightingGuid = FGuid::NewGuid();
+			// Construct landscape static mesh
+			FString LandscapeMeshAssetName = TEXT("SM_") + LandscapeBaseAssetName;
+			UStaticMesh* StaticMesh = new(LODPackage, *LandscapeMeshAssetName, RF_Public|RF_Standalone) UStaticMesh(FPostConstructInitializeProperties());
+			StaticMesh->InitResources();
+			{
+				FString OutputPath = StaticMesh->GetPathName();
 
-			// Set it to use textured lightmaps. Note that Build Lighting will do the error-checking (texcoordindex exists for all LODs, etc).
-			StaticMesh->LightMapResolution = 128;
-			StaticMesh->LightMapCoordinateIndex = 1;
+				// make sure it has a new lighting guid
+				StaticMesh->LightingGuid = FGuid::NewGuid();
 
-			FStaticMeshSourceModel* SrcModel = new (StaticMesh->SourceModels) FStaticMeshSourceModel();
-			/*Don't allow the engine to recalculate normals*/
-			SrcModel->BuildSettings.bRecomputeNormals = false;
-			SrcModel->BuildSettings.bRecomputeTangents = false;
-			SrcModel->BuildSettings.bRemoveDegenerates = false;
-			SrcModel->BuildSettings.bUseFullPrecisionUVs = false;
-			SrcModel->RawMeshBulkData->SaveRawMesh(LandscapeRawMesh);
+				// Set it to use textured lightmaps. Note that Build Lighting will do the error-checking (texcoordindex exists for all LODs, etc).
+				StaticMesh->LightMapResolution = 128;
+				StaticMesh->LightMapCoordinateIndex = 1;
 
-			//Assign the proxy material to the static mesh
-			StaticMesh->Materials.Add(StaticLandscapeMaterial);
+				FStaticMeshSourceModel* SrcModel = new (StaticMesh->SourceModels) FStaticMeshSourceModel();
+				/*Don't allow the engine to recalculate normals*/
+				SrcModel->BuildSettings.bRecomputeNormals = false;
+				SrcModel->BuildSettings.bRecomputeTangents = false;
+				SrcModel->BuildSettings.bRemoveDegenerates = false;
+				SrcModel->BuildSettings.bUseFullPrecisionUVs = false;
+				SrcModel->RawMeshBulkData->SaveRawMesh(LandscapeRawMesh);
 
-			StaticMesh->Build();
-			StaticMesh->PostEditChange();
-		}
+				//Assign the proxy material to the static mesh
+				StaticMesh->Materials.Add(StaticLandscapeMaterial);
 
-		AssetsToSpawn.Add(StaticMesh);
-		AssetsToSpawnTransform.Add(FTransform(LandscapeWorldLocation - ActorsOffset));
+				StaticMesh->Build();
+				StaticMesh->PostEditChange();
+			}
 
-		GeneratedAssets.Add(StaticMesh);
-		GeneratedAssets.Add(StaticLandscapeMaterial);
-		LandscapeActorIndex++;
-	}
+			AssetsToSpawn.Add(StaticMesh);
+			AssetsToSpawnTransform.Add(FTransform(LandscapeWorldLocation - ActorsOffset));
 
-	// Restore level original visibility
-	if (!bVisibleLevel)
-	{
-		InLevelModel->SetVisible(false);
-		GetWorld()->WorldComposition->bTemporallyDisableOriginTracking = false;
-	}
+			LandscapeActorIndex++;
+		}
 
-	// Save generated assets
-	TArray<UPackage*> AssetsPackagesToSave;
-	for (UObject* Asset : GeneratedAssets)
-	{
-		if (Asset->IsAsset())
+		// Restore level original visibility
+		if (!bVisibleLevel)
 		{
-			AssetsPackagesToSave.AddUnique(Asset->GetOutermost());
+			TileModel->SetVisible(false);
+			GetWorld()->WorldComposition->bTemporallyDisableOriginTracking = false;
 		}
-	}
 	
-	FEditorFileUtils::PromptForCheckoutAndSave(AssetsPackagesToSave, false, false);
-
-	// Create new level and spawn generated assets in it
-	if (AssetsToSpawn.Num())
-	{
-		// Create a new world - so we can 'borrow' its level
-		UWorld* LODWorld = UWorld::CreateWorld(EWorldType::None, false);
-
-		for (int32 AssetIdx = 0; AssetIdx < AssetsToSpawn.Num(); ++AssetIdx)
+		// Create new level and spawn generated assets in it
+		if (AssetsToSpawn.Num())
 		{
-			FVector Location = AssetsToSpawnTransform[AssetIdx].GetLocation();
-			FRotator Rotation(ForceInit);
-			AStaticMeshActor* MeshActor = LODWorld->SpawnActor<AStaticMeshActor>(Location, Rotation);
-			MeshActor->StaticMeshComponent->StaticMesh = AssetsToSpawn[AssetIdx];
-			MeshActor->SetActorLabel(AssetsToSpawn[AssetIdx]->GetName());
-		}
+			// Create a new world
+			UWorld* LODWorld = UWorld::CreateWorld(EWorldType::None, false, FPackageName::GetShortFName(LODPackage->GetFName()), LODPackage);
+			LODWorld->SetFlags(RF_Standalone);
 
-		FEditorFileUtils::SaveLevel(LODWorld->PersistentLevel, LODLevelFilename);
-		// Destroy the new world we created and collect the garbage
-		LODWorld->DestroyWorld(false);
-		CollectGarbage(GARBAGE_COLLECTION_KEEPFLAGS);
-	}
-
-	//Update the asset registry that a new static mash and material has been created
-	if (GeneratedAssets.Num())
-	{
-		FAssetRegistryModule& AssetRegistry = FModuleManager::Get().LoadModuleChecked<FAssetRegistryModule>("AssetRegistry");
-		int32 AssetCount = GeneratedAssets.Num();
-		for (int32 AssetIndex = 0; AssetIndex < AssetCount; AssetIndex++)
-		{
-			AssetRegistry.AssetCreated(GeneratedAssets[AssetIndex]);
-			GEditor->BroadcastObjectReimported(GeneratedAssets[AssetIndex]);
+			for (int32 AssetIdx = 0; AssetIdx < AssetsToSpawn.Num(); ++AssetIdx)
+			{
+				FVector Location = AssetsToSpawnTransform[AssetIdx].GetLocation();
+				FRotator Rotation(ForceInit);
+				AStaticMeshActor* MeshActor = LODWorld->SpawnActor<AStaticMeshActor>(Location, Rotation);
+				MeshActor->StaticMeshComponent->StaticMesh = AssetsToSpawn[AssetIdx];
+				MeshActor->SetActorLabel(AssetsToSpawn[AssetIdx]->GetName());
+			}
+		
+			// Save generated level
+			if (FEditorFileUtils::PromptToCheckoutLevels(false, LODWorld->PersistentLevel))
+			{
+				FEditorFileUtils::SaveLevel(LODWorld->PersistentLevel, *LODLevelFileName);
+			}
+			
+			// Destroy the new world we created and collect the garbage
+			LODWorld->ClearFlags(RF_Standalone);
+			LODWorld->DestroyWorld(false);
+			CollectGarbage(GARBAGE_COLLECTION_KEEPFLAGS);
 		}
-
-		//Also notify the content browser that the new assets exists
-		FContentBrowserModule& ContentBrowserModule = FModuleManager::Get().LoadModuleChecked<FContentBrowserModule>("ContentBrowser");
-		ContentBrowserModule.Get().SyncBrowserToAssets(GeneratedAssets, true);
 	}
-	
+
 	// Rescan world root
 	PopulateLevelsList();
-
 	GWarn->EndSlowTask();	
+	
 	return true;
 }
 
diff --git a/Engine/Source/Editor/WorldBrowser/Private/Tiles/WorldTileCollectionModel.h b/Engine/Source/Editor/WorldBrowser/Private/Tiles/WorldTileCollectionModel.h
--- a/Engine/Source/Editor/WorldBrowser/Private/Tiles/WorldTileCollectionModel.h
+++ b/Engine/Source/Editor/WorldBrowser/Private/Tiles/WorldTileCollectionModel.h
@@ -232,11 +232,11 @@ public:
 	bool HasGenerateLODLevelSupport() const;
 	
 	/** 
-	 * Generates LOD level from a specified level. Level has to be loaded.
-	 * Currently all static meshes will be merged into one proxy mesh using Simplygon ProxyLOD
-	 * Landscape actors will be converted into static meshes and simplified using mesh reduction interface
+	 * Generates simplified versions of a specified levels. Levels has to be loaded.
+	 * Currently all static meshes found inside one level will be merged into one proxy mesh using Simplygon ProxyLOD
+	 * Landscape actors will be converted into static meshes using highest landscape LOD entry
 	 */	
-	bool GenerateLODLevel(TSharedPtr<FLevelModel> InLevelModel, int32 TargetLODIndex);
+	bool GenerateLODLevels(FLevelModelList InLevelList, int32 TargetLODIndex);
 
 	/** Assign selected levels to current layer */
 	void AssignSelectedLevelsToLayer_Executed(FWorldTileLayer InLayer);
diff --git a/Engine/Source/Editor/WorldBrowser/Private/Tiles/WorldTileDetailsCustomization.cpp b/Engine/Source/Editor/WorldBrowser/Private/Tiles/WorldTileDetailsCustomization.cpp
--- a/Engine/Source/Editor/WorldBrowser/Private/Tiles/WorldTileDetailsCustomization.cpp
+++ b/Engine/Source/Editor/WorldBrowser/Private/Tiles/WorldTileDetailsCustomization.cpp
@@ -277,7 +277,8 @@ void FTileLODEntryDetailsCustomization::CustomizeChildren(TSharedRef<IPropertyHa
 	ChildBuilder.AddChildProperty(LODIndexHandle.ToSharedRef())
 		.Visibility(EVisibility::Hidden);
 
-	ChildBuilder.AddChildProperty(DistanceProperty.ToSharedRef());
+	ChildBuilder.AddChildProperty(DistanceProperty.ToSharedRef())
+		.IsEnabled(TAttribute<bool>::Create(TAttribute<bool>::FGetter::CreateSP(this, &FTileLODEntryDetailsCustomization::IsLODDistanceEnabled)));
 	// Reduction settings available if editor supports LOD levels generation
 	ChildBuilder.AddChildProperty(DetailsPercentage.ToSharedRef())
 		.IsEnabled(TAttribute<bool>::Create(TAttribute<bool>::FGetter::CreateSP(this, &FTileLODEntryDetailsCustomization::IsGenerateTileEnabled)));
@@ -295,10 +296,7 @@ FReply FTileLODEntryDetailsCustomization::OnGenerateTile()
 		if (PinnedWorldModel.IsValid())
 		{
 			FLevelModelList LevelsList = PinnedWorldModel->GetSelectedLevels();
-			for (auto Level : LevelsList)
-			{
-				PinnedWorldModel->GenerateLODLevel(Level, LODIndex);
-			}
+			PinnedWorldModel->GenerateLODLevels(LevelsList, LODIndex);
 		}
 	}
 	
@@ -310,7 +308,18 @@ bool FTileLODEntryDetailsCustomization::IsGenerateTileEnabled() const
 	TSharedPtr<FWorldTileCollectionModel> PinnedWorldModel = WorldModel.Pin();
 	if (PinnedWorldModel.IsValid())
 	{
-		return PinnedWorldModel->HasGenerateLODLevelSupport();
+		return PinnedWorldModel->HasGenerateLODLevelSupport() && PinnedWorldModel->AreAnySelectedLevelsLoaded();
+	}
+	
+	return false;
+}
+
+bool FTileLODEntryDetailsCustomization::IsLODDistanceEnabled() const
+{
+	TSharedPtr<FWorldTileCollectionModel> PinnedWorldModel = WorldModel.Pin();
+	if (PinnedWorldModel.IsValid())
+	{
+		return PinnedWorldModel->AreAnySelectedLevelsEditable();
 	}
 	
 	return false;
diff --git a/Engine/Source/Editor/WorldBrowser/Private/Tiles/WorldTileDetailsCustomization.h b/Engine/Source/Editor/WorldBrowser/Private/Tiles/WorldTileDetailsCustomization.h
--- a/Engine/Source/Editor/WorldBrowser/Private/Tiles/WorldTileDetailsCustomization.h
+++ b/Engine/Source/Editor/WorldBrowser/Private/Tiles/WorldTileDetailsCustomization.h
@@ -62,6 +62,7 @@ public:
 private:
 	FReply OnGenerateTile();
 	bool IsGenerateTileEnabled() const;
+	bool IsLODDistanceEnabled() const;
 
 private:
 	TWeakPtr<FWorldTileCollectionModel> WorldModel;
diff --git a/Engine/Source/Runtime/Engine/Private/WorldComposition.cpp b/Engine/Source/Runtime/Engine/Private/WorldComposition.cpp
--- a/Engine/Source/Runtime/Engine/Private/WorldComposition.cpp
+++ b/Engine/Source/Runtime/Engine/Private/WorldComposition.cpp
@@ -227,7 +227,7 @@ void UWorldComposition::Rescan()
 		if (TileLODList.Num())
 		{
 			Tile.LODPackageNames.SetNum(WORLDTILE_LOD_MAX_INDEX);
-			FString TilePath = FPackageName::GetLongPackagePath(TilePackageName);
+			FString TilePath = FPackageName::GetLongPackagePath(TilePackageName) + TEXT("/");
 			for (const auto& TileLOD : TileLODList)
 			{
 				// LOD tiles should be in the same directory or in nested directory
