World Browser: Fixed case when user was able to move selected actors to a hidden level, which was causing unexpected results

diff --git a/Engine/Source/Editor/UnrealEd/Private/EditorServer.cpp b/Engine/Source/Editor/UnrealEd/Private/EditorServer.cpp
--- a/Engine/Source/Editor/UnrealEd/Private/EditorServer.cpp
+++ b/Engine/Source/Editor/UnrealEd/Private/EditorServer.cpp
@@ -2571,6 +2571,12 @@ void UEditorEngine::DoMoveSelectedActorsToLevel( ULevel* InDestLevel )
 		return;
 	}
 
+	if (GetSelectedActorCount() <= 0)
+	{
+		// Nothing to move, probably source level was hidden and actors lost selection mark
+		return;
+	}
+
 	// Start the transaction
 	GEditor->Trans->Begin( NULL, NSLOCTEXT("UnrealEd", "MoveSelectedActorsToSelectedLevel", "Move Actors To Level"));
 	
diff --git a/Engine/Source/Editor/WorldBrowser/Private/LevelCollectionModel.cpp b/Engine/Source/Editor/WorldBrowser/Private/LevelCollectionModel.cpp
--- a/Engine/Source/Editor/WorldBrowser/Private/LevelCollectionModel.cpp
+++ b/Engine/Source/Editor/WorldBrowser/Private/LevelCollectionModel.cpp
@@ -72,7 +72,7 @@ void FLevelCollectionModel::BindCommands()
 	
 	ActionList.MapAction(Commands.MoveActorsToSelected,
 		FExecuteAction::CreateSP(this, &FLevelCollectionModel::MoveActorsToSelected_Executed),
-		FCanExecuteAction::CreateSP(this, &FLevelCollectionModel::IsSelectedLevelEditable));
+		FCanExecuteAction::CreateSP(this, &FLevelCollectionModel::CanMoveSelectedActorsToLevel));
 		
 	ActionList.MapAction(Commands.World_SaveSelectedLevels,
 		FExecuteAction::CreateSP(this, &FLevelCollectionModel::SaveSelectedLevels_Executed),
@@ -1483,6 +1483,11 @@ void FLevelCollectionModel::MoveActorsToSelected_Executed()
 	RequestUpdateAllLevels();
 }
 
+bool FLevelCollectionModel::CanMoveSelectedActorsToLevel() const
+{
+	return (IsOneLevelSelected() && AreAllSelectedLevelsEditableAndVisible() && Editor->GetSelectedActorCount() > 0);
+}
+
 void FLevelCollectionModel::SelectActors_Executed()
 {
 	//first clear any existing actor selection
diff --git a/Engine/Source/Editor/WorldBrowser/Private/LevelCollectionModel.h b/Engine/Source/Editor/WorldBrowser/Private/LevelCollectionModel.h
--- a/Engine/Source/Editor/WorldBrowser/Private/LevelCollectionModel.h
+++ b/Engine/Source/Editor/WorldBrowser/Private/LevelCollectionModel.h
@@ -291,6 +291,7 @@ protected:
 
 	/** Moves the selected actors to this level */
 	void MoveActorsToSelected_Executed();
+	bool CanMoveSelectedActorsToLevel() const;
 
 	/** Saves selected levels */
 	void SaveSelectedLevels_Executed();
