Level and world browsers Move Selected Actor To Level command is only enabled if the the selected actors arenÆt in the selected level.
#TTP 335685 - EDITOR: Level Browser: Allowed to click "move selected actor to level" on an actor that is already in the selected level

#branch UE4

#proj Editor.WorldBrowser
#proj Editor.Levels

#change Added a new method IsValidMoveActorsToLevel to both FLevelCollectionModel and FLevelCollectionViewModel, this is used to enable the Move Actor to level command for both cases. It scans the selected actors and checks their levels against the selected levels. If any actor is contained within any of the selected levels it disables the control. It also disables the control should there be no actors selected. If both of these cases pass it simply checks if the level is editable.

diff --git a/Engine/Source/Editor/Levels/Private/LevelCollectionViewModel.cpp b/Engine/Source/Editor/Levels/Private/LevelCollectionViewModel.cpp
--- a/Engine/Source/Editor/Levels/Private/LevelCollectionViewModel.cpp
+++ b/Engine/Source/Editor/Levels/Private/LevelCollectionViewModel.cpp
@@ -25,6 +25,7 @@ FLevelCollectionViewModel::FLevelCollectionViewModel( const TWeakObjectPtr< UEdi
 	, Editor( InEditor )
 	, CurrentWorld( NULL )
 	, AddedLevelStreamingClass( ULevelStreamingKismet::StaticClass() )
+	, bSelectionHasChanged( true )
 {
 	OnResetLevels();
 }
@@ -68,6 +69,9 @@ void FLevelCollectionViewModel::Initialize()
 			SetCurrentWorld( GEditor->GetEditorWorldContext().World() );
 		}
 	}
+	
+	USelection::SelectionChangedEvent.AddSP(this, &FLevelCollectionViewModel::OnActorSelectionChanged);
+	SelectionChanged.AddSP(this, &FLevelCollectionViewModel::OnActorOrLevelSelectionChanged);
 }
 
 void FLevelCollectionViewModel::Tick( float DeltaTime )
@@ -95,7 +99,7 @@ void FLevelCollectionViewModel::BindCommands()
 	
 	ActionList.MapAction( Commands.MoveActorsToSelected,
 		FExecuteAction::CreateSP( this, &FLevelCollectionViewModel::MoveActorsToSelected_Executed  ),
-		FCanExecuteAction::CreateSP( this, &FLevelCollectionViewModel::IsSelectedLevelUnlocked ) );
+		FCanExecuteAction::CreateSP(this, &FLevelCollectionViewModel::IsValidMoveActorsToLevel));
 
 	//invalid selected levels
 	ActionList.MapAction( Commands.FixUpInvalidReference,
@@ -994,6 +998,20 @@ bool FLevelCollectionViewModel::AreSelectedLevelsUnlockedAndNotPersistent() cons
 	return true;
 }
 
+bool FLevelCollectionViewModel::AreAllSelectedLevelsUnlockedAndVisible() const
+{
+	for (const auto& Level : SelectedLevels)
+	{
+		if (Level->IsLocked() == true ||
+			Level->IsVisible() == false) 
+		{
+			return false;
+		}
+	}
+
+	return SelectedLevels.Num() > 0;
+}
+
 bool FLevelCollectionViewModel::AreActorsSelected() const
 {
 	return Editor->GetSelectedActorCount() > 0;
@@ -2107,4 +2125,53 @@ void FLevelCollectionViewModel::AddExistingLevelFromAssetPicker(const TArray<FAs
 	}
 }
 
+bool FLevelCollectionViewModel::IsValidMoveActorsToLevel()
+{
+	static bool bCachedIsValidActorMoveResult = false;
+	if (bSelectionHasChanged)
+	{
+		bSelectionHasChanged = false;
+		USelection* SelectedActors = GEditor->GetSelectedActors();
+
+		// you cant move no selected actors to a level
+		if (SelectedActors->Num() == 0)
+		{
+			bCachedIsValidActorMoveResult = false;
+			return false;
+		}
+
+		// are any of the selected actors in the selected levels
+		for (FSelectionIterator Iter(*SelectedActors); Iter; ++Iter)
+		{
+			AActor* Actor = CastChecked<AActor>(*Iter);
+			if (Actor != nullptr)
+			{
+				const ULevel* ActorsLevel = Actor->GetLevel();
+				for (TSharedPtr<FLevelViewModel> SelectedLevel : SelectedLevels)
+				{
+					if (SelectedLevel->GetLevel().Get() == ActorsLevel)
+					{
+						bCachedIsValidActorMoveResult = false;
+						return false;
+					}
+				}
+			}
+		}
+		bCachedIsValidActorMoveResult = true;
+	}
+
+	// if non of the selected actors are in the level, just check the level is unlocked
+	return bCachedIsValidActorMoveResult && AreAllSelectedLevelsUnlockedAndVisible();
+}
+
+void FLevelCollectionViewModel::OnActorSelectionChanged(UObject* obj)
+{
+	OnActorOrLevelSelectionChanged();
+}
+
+void FLevelCollectionViewModel::OnActorOrLevelSelectionChanged()
+{
+	bSelectionHasChanged = true;
+}
+
 #undef LOCTEXT_NAMESPACE
diff --git a/Engine/Source/Editor/Levels/Private/LevelCollectionViewModel.h b/Engine/Source/Editor/Levels/Private/LevelCollectionViewModel.h
--- a/Engine/Source/Editor/Levels/Private/LevelCollectionViewModel.h
+++ b/Engine/Source/Editor/Levels/Private/LevelCollectionViewModel.h
@@ -309,6 +309,9 @@ private:
 	/** @return whether the currently selected levels are unlocked and do not contain the persistent level */
 	bool AreSelectedLevelsUnlockedAndNotPersistent() const;
 
+	/** @return true if we have selected levels and they are all editable and visible, */
+	bool AreAllSelectedLevelsUnlockedAndVisible() const;
+
 	/** @return	whether -at least- one actor is selected */
 	bool AreActorsSelected() const;
 	
@@ -458,7 +461,15 @@ private:
 
 	/** Toggle all read-only levels */
 	void ToggleReadOnlyLevels_Executed();
+	
+	/** @return whether moving the selected actors to the selected level is a valid action */
+	bool IsValidMoveActorsToLevel();
+
+	/** delegate used to pickup when the selection has changed */
+	void OnActorSelectionChanged(UObject* obj);
 
+	/** Sets a flag to re-cache whether the selected actors move to the selected level is valid */
+	void OnActorOrLevelSelectionChanged();
 private:
 
 	/** true if the LevelsView is in the middle of refreshing */
@@ -527,6 +538,8 @@ private:
 	/** true if Source Control options are generally available. */
 	bool bCanExecuteSCC;
 	
+	/** Flag for whether the selection of levels or actors has changed */
+	bool bSelectionHasChanged;
 };
 
 
diff --git a/Engine/Source/Editor/WorldBrowser/Private/LevelCollectionModel.cpp b/Engine/Source/Editor/WorldBrowser/Private/LevelCollectionModel.cpp
--- a/Engine/Source/Editor/WorldBrowser/Private/LevelCollectionModel.cpp
+++ b/Engine/Source/Editor/WorldBrowser/Private/LevelCollectionModel.cpp
@@ -23,6 +23,7 @@ FLevelCollectionModel::FLevelCollectionModel(const TWeakObjectPtr<UEditorEngine>
 	, bCanExecuteSCCOpenForAdd(false)
 	, bCanExecuteSCCCheckIn(false)
 	, bCanExecuteSCC(false)
+	, bSelectionHasChanged(true)
 {
 }
 
@@ -50,7 +51,9 @@ void FLevelCollectionModel::Initialize()
 	FEditorSupportDelegates::RedrawAllViewports.AddSP(this, &FLevelCollectionModel::OnRedrawAllViewports);
 	Editor->OnLevelActorAdded().AddSP( this, &FLevelCollectionModel::OnLevelActorAdded);
 	Editor->OnLevelActorDeleted().AddSP( this, &FLevelCollectionModel::OnLevelActorDeleted);
-	
+	USelection::SelectionChangedEvent.AddSP(this, &FLevelCollectionModel::OnActorSelectionChanged);
+	SelectionChanged.AddSP(this, &FLevelCollectionModel::OnActorOrLevelSelectionChanged);
+
 	PopulateLevelsList();
 }
 
@@ -72,7 +75,7 @@ void FLevelCollectionModel::BindCommands()
 	
 	ActionList.MapAction(Commands.MoveActorsToSelected,
 		FExecuteAction::CreateSP(this, &FLevelCollectionModel::MoveActorsToSelected_Executed),
-		FCanExecuteAction::CreateSP(this, &FLevelCollectionModel::CanMoveSelectedActorsToLevel));
+		FCanExecuteAction::CreateSP(this, &FLevelCollectionModel::IsValidMoveActorsToLevel));
 		
 	ActionList.MapAction(Commands.World_SaveSelectedLevels,
 		FExecuteAction::CreateSP(this, &FLevelCollectionModel::SaveSelectedLevels_Executed),
@@ -1483,11 +1486,6 @@ void FLevelCollectionModel::MoveActorsToSelected_Executed()
 	RequestUpdateAllLevels();
 }
 
-bool FLevelCollectionModel::CanMoveSelectedActorsToLevel() const
-{
-	return (IsOneLevelSelected() && AreAllSelectedLevelsEditableAndVisible() && Editor->GetSelectedActorCount() > 0);
-}
-
 void FLevelCollectionModel::SelectActors_Executed()
 {
 	//first clear any existing actor selection
@@ -1728,5 +1726,54 @@ void FLevelCollectionModel::CacheCanExecuteSourceControlVars() const
 	}
 }
 
+bool FLevelCollectionModel::IsValidMoveActorsToLevel() 
+{
+	static bool bCachedIsValidActorMoveResult = false;
+	if (bSelectionHasChanged)
+	{
+		bSelectionHasChanged = false;
+		USelection* SelectedActors = GEditor->GetSelectedActors();
+		// you cant move no selected actors to a level
+		if (SelectedActors->Num() == 0)
+		{
+			bCachedIsValidActorMoveResult = false;
+			return false;
+		}
+
+		// are any of the selected actors in the selected levels
+		for (FSelectionIterator Iter(*SelectedActors); Iter; ++Iter)
+		{
+			AActor* Actor = CastChecked<AActor>(*Iter);
+			if (Actor != nullptr)
+			{
+				const ULevel* ActorsLevel = Actor->GetLevel();
+				for (TSharedPtr<FLevelModel> SelectedLevel : SelectedLevelsList)
+				{
+					if (SelectedLevel->GetLevelObject() == ActorsLevel)
+					{
+						bCachedIsValidActorMoveResult = false;
+						return false;
+					}
+				}
+			}
+		}
+
+		bCachedIsValidActorMoveResult = true;
+	}
+			
+	// if non of the selected actors are in the level, just check the level is unlocked
+	return bCachedIsValidActorMoveResult && AreAllSelectedLevelsEditableAndVisible();
+}
+
+void FLevelCollectionModel::OnActorSelectionChanged(UObject* obj)
+{
+	OnActorOrLevelSelectionChanged();
+}
+
+void FLevelCollectionModel::OnActorOrLevelSelectionChanged()
+{
+	bSelectionHasChanged = true;
+}
+
 
 #undef LOCTEXT_NAMESPACE
diff --git a/Engine/Source/Editor/WorldBrowser/Private/LevelCollectionModel.h b/Engine/Source/Editor/WorldBrowser/Private/LevelCollectionModel.h
--- a/Engine/Source/Editor/WorldBrowser/Private/LevelCollectionModel.h
+++ b/Engine/Source/Editor/WorldBrowser/Private/LevelCollectionModel.h
@@ -205,6 +205,15 @@ public:
 	/** @return	whether at least one actor is selected */
 	bool AreActorsSelected() const;
 
+	/** @return whether moving the selected actors to the selected level is a valid action */
+	bool IsValidMoveActorsToLevel();
+
+	/** delegate used to pickup when the selection has changed */
+	void OnActorSelectionChanged(UObject* obj);
+
+	/** Sets a flag to re-cache whether the selected actors move to the selected level is valid */
+	void OnActorOrLevelSelectionChanged();
+
 	/** @return	whether 'display paths' is enabled */
 	bool GetDisplayPathsState() const;
 
@@ -291,7 +300,6 @@ protected:
 
 	/** Moves the selected actors to this level */
 	void MoveActorsToSelected_Executed();
-	bool CanMoveSelectedActorsToLevel() const;
 
 	/** Saves selected levels */
 	void SaveSelectedLevels_Executed();
@@ -504,6 +512,9 @@ protected:
 
 	/** true if Source Control options are generally available. */
 	mutable bool						bCanExecuteSCC;
+
+	/** Flag for whether the selection of levels or actors has changed */
+	bool								bSelectionHasChanged;
 };
 
 //
