World Browser: Added rotation to a simulation view position indicator. Changed world bounds color from red to yellow

diff --git a/Engine/Source/Editor/WorldBrowser/Private/LevelCollectionModel.cpp b/Engine/Source/Editor/WorldBrowser/Private/LevelCollectionModel.cpp
--- a/Engine/Source/Editor/WorldBrowser/Private/LevelCollectionModel.cpp
+++ b/Engine/Source/Editor/WorldBrowser/Private/LevelCollectionModel.cpp
@@ -701,9 +701,9 @@ void FLevelCollectionModel::CustomizeFileMainMenu(FMenuBuilder& InMenuBuilder) c
 	}
 }
 
-FVector FLevelCollectionModel::GetObserverPosition() const
+FMatrix FLevelCollectionModel::GetObserverViewMatrix() const
 {
-	return FVector::ZeroVector;
+	return FMatrix::Identity;
 }
 
 bool FLevelCollectionModel::CompareLevelsZOrder(TSharedPtr<FLevelModel> InA, TSharedPtr<FLevelModel> InB) const
diff --git a/Engine/Source/Editor/WorldBrowser/Private/LevelCollectionModel.h b/Engine/Source/Editor/WorldBrowser/Private/LevelCollectionModel.h
--- a/Engine/Source/Editor/WorldBrowser/Private/LevelCollectionModel.h
+++ b/Engine/Source/Editor/WorldBrowser/Private/LevelCollectionModel.h
@@ -127,8 +127,8 @@ public:
 	/**	Customize 'File' section in main menu  */
 	virtual void CustomizeFileMainMenu(FMenuBuilder& InMenuBuilder) const;
 
-	/**	@return	Observer position in the world, usually camera position */
-	virtual FVector GetObserverPosition() const;
+	/**	@return	Observer view matrix in the world, usually camera view */
+	virtual FMatrix GetObserverViewMatrix() const;
 
 	/**	Compares 2 levels by Z order */
 	virtual bool CompareLevelsZOrder(TSharedPtr<FLevelModel> InA, TSharedPtr<FLevelModel> InB) const;
diff --git a/Engine/Source/Editor/WorldBrowser/Private/Tiles/SWorldComposition.cpp b/Engine/Source/Editor/WorldBrowser/Private/Tiles/SWorldComposition.cpp
--- a/Engine/Source/Editor/WorldBrowser/Private/Tiles/SWorldComposition.cpp
+++ b/Engine/Source/Editor/WorldBrowser/Private/Tiles/SWorldComposition.cpp
@@ -318,7 +318,7 @@ public:
 			FVector2D PaintPosition = GraphCoordToPanelCoord(FVector2D::ZeroVector) - (PaintSize*0.5f);
 			FPaintGeometry EditableArea = AllottedGeometry.ToPaintGeometry(PaintPosition, PaintSize);
 			EditableArea.DrawScale = 0.2f;
-			FLinearColor PaintColor = FLinearColor::Red;
+			FLinearColor PaintColor = FLinearColor::Yellow;
 			PaintColor.A = 0.4f;
 
 			FSlateDrawElement::MakeBox(
@@ -352,21 +352,28 @@ public:
 			// Draw a current camera position
 			{
 				const FSlateBrush* CameraImage = FEditorStyle::GetBrush(TEXT("WorldBrowser.SimulationViewPositon"));
-				FVector ObserverPosition = WorldModel->GetObserverPosition();
+				FMatrix ObserverViewToWorld = WorldModel->GetObserverViewMatrix().Inverse();
+				FVector ObserverPosition	= ObserverViewToWorld.GetOrigin();
+				FRotator ObserverRotation	= ObserverViewToWorld.Rotator();
+								
 				FVector2D ObserverPositionScreen = GraphCoordToPanelCoord(FVector2D(ObserverPosition.X, ObserverPosition.Y));
 
 				FPaintGeometry PaintGeometry = AllottedGeometry.ToPaintGeometry(
 					ObserverPositionScreen - CameraImage->ImageSize*0.5f, 
 					CameraImage->ImageSize
 					);
-						
-				FSlateDrawElement::MakeBox(
+
+				FSlateDrawElement::MakeRotatedBox(
 					OutDrawElements,
 					++LayerId,
 					PaintGeometry,
 					CameraImage,
-					MyClippingRect
-				);
+					MyClippingRect,
+					ESlateDrawEffect::None,
+					FMath::DegreesToRadians(ObserverRotation.Yaw - 90.f), // 0 degrees pointing along +Y axis, arrow on the image pointing along +X axis
+					CameraImage->ImageSize*0.5f,
+					FSlateDrawElement::RelativeToElement
+					);
 			}
 		}
 
diff --git a/Engine/Source/Editor/WorldBrowser/Private/Tiles/WorldTileCollectionModel.cpp b/Engine/Source/Editor/WorldBrowser/Private/Tiles/WorldTileCollectionModel.cpp
--- a/Engine/Source/Editor/WorldBrowser/Private/Tiles/WorldTileCollectionModel.cpp
+++ b/Engine/Source/Editor/WorldBrowser/Private/Tiles/WorldTileCollectionModel.cpp
@@ -517,16 +517,16 @@ void FWorldTileCollectionModel::CustomizeFileMainMenu(FMenuBuilder& InMenuBuilde
 	InMenuBuilder.EndSection();
 }
 
-FVector FWorldTileCollectionModel::GetObserverPosition() const
+FMatrix FWorldTileCollectionModel::GetObserverViewMatrix() const
 {
 	UWorld*	SimulationWorld = GetSimulationWorld();
 	if (SimulationWorld && SimulationWorld->WorldComposition)
 	{
-		return SimulationWorld->WorldComposition->LastViewLocation;
+		return SimulationWorld->WorldComposition->LastWorldToViewMatrix;
 	}
 	else
 	{
-		return FVector::ZeroVector;
+		return FMatrix::Identity;
 	}
 }
 
diff --git a/Engine/Source/Editor/WorldBrowser/Private/Tiles/WorldTileCollectionModel.h b/Engine/Source/Editor/WorldBrowser/Private/Tiles/WorldTileCollectionModel.h
--- a/Engine/Source/Editor/WorldBrowser/Private/Tiles/WorldTileCollectionModel.h
+++ b/Engine/Source/Editor/WorldBrowser/Private/Tiles/WorldTileCollectionModel.h
@@ -34,7 +34,7 @@ public:
 	virtual void BuildGridMenu(FMenuBuilder& InMenuBuilder) const override;
 	virtual void BuildHierarchyMenu(FMenuBuilder& InMenuBuilder) const override;
 	virtual void CustomizeFileMainMenu(FMenuBuilder& InMenuBuilder) const override;
-	virtual FVector GetObserverPosition() const override;
+	virtual FMatrix GetObserverViewMatrix() const override;
 	virtual bool CompareLevelsZOrder(TSharedPtr<FLevelModel> InA, TSharedPtr<FLevelModel> InB) const override;
 	virtual void RegisterDetailsCustomization(class FPropertyEditorModule& PropertyModule, TSharedPtr<class IDetailsView> InDetailsView)  override;
 	virtual void UnregisterDetailsCustomization(class FPropertyEditorModule& PropertyModule, TSharedPtr<class IDetailsView> InDetailsView) override;
diff --git a/Engine/Source/Runtime/Engine/Classes/Engine/WorldComposition.h b/Engine/Source/Runtime/Engine/Classes/Engine/WorldComposition.h
--- a/Engine/Source/Runtime/Engine/Classes/Engine/WorldComposition.h
+++ b/Engine/Source/Runtime/Engine/Classes/Engine/WorldComposition.h
@@ -168,8 +168,8 @@ private:
 
 public:
 #if WITH_EDITOR
-	// Last location from where streaming state was updated
-	mutable FVector				LastViewLocation;
+	// Last view from where streaming state was updated
+	mutable FMatrix				LastWorldToViewMatrix;
 
 	// Hack for a World Browser to be able to temporally show hidden levels 
 	// regardless of current world origin and without offsetting them temporally 
diff --git a/Engine/Source/Runtime/Engine/Private/WorldComposition.cpp b/Engine/Source/Runtime/Engine/Private/WorldComposition.cpp
--- a/Engine/Source/Runtime/Engine/Private/WorldComposition.cpp
+++ b/Engine/Source/Runtime/Engine/Private/WorldComposition.cpp
@@ -526,10 +526,6 @@ void UWorldComposition::UpdateStreamingState(const FVector& InLocation)
 	{
 		CommitTileStreamingState(OwningWorld, Level.TileIdx, true, true, Level.LODIndex);
 	}
-
-#if WITH_EDITOR
-	LastViewLocation = InLocation;
-#endif //WITH_EDITOR
 }
 
 void UWorldComposition::UpdateStreamingState(FSceneViewFamily* ViewFamily)
@@ -546,6 +542,10 @@ void UWorldComposition::UpdateStreamingState(FSceneViewFamily* ViewFamily)
 		}
 
 		ViewLocation/= ViewFamily->Views.Num();
+
+#if WITH_EDITOR
+	LastWorldToViewMatrix = ViewFamily->Views[0]->ViewMatrices.ViewMatrix;
+#endif //WITH_EDITOR
 	}
 	
 	UpdateStreamingState(ViewLocation);
