World Browser: Fixed crash related to simulation view position

diff --git a/Engine/Source/Editor/WorldBrowser/Private/Tiles/SWorldComposition.cpp b/Engine/Source/Editor/WorldBrowser/Private/Tiles/SWorldComposition.cpp
--- a/Engine/Source/Editor/WorldBrowser/Private/Tiles/SWorldComposition.cpp
+++ b/Engine/Source/Editor/WorldBrowser/Private/Tiles/SWorldComposition.cpp
@@ -352,7 +352,7 @@ public:
 			// Draw a current camera position
 			{
 				const FSlateBrush* CameraImage = FEditorStyle::GetBrush(TEXT("WorldBrowser.SimulationViewPositon"));
-				FMatrix ObserverViewToWorld = WorldModel->GetObserverViewMatrix().Inverse();
+				FMatrix ObserverViewToWorld = WorldModel->GetObserverViewMatrix().InverseSafe();
 				FVector ObserverPosition	= ObserverViewToWorld.GetOrigin();
 				FRotator ObserverRotation	= ObserverViewToWorld.Rotator();
 								
