Fixed crash when switching levels selection using "Inspect level" combo box
TTP# 339980

diff --git a/Engine/Source/Editor/WorldBrowser/Private/SWorldDetails.cpp b/Engine/Source/Editor/WorldBrowser/Private/SWorldDetails.cpp
--- a/Engine/Source/Editor/WorldBrowser/Private/SWorldDetails.cpp
+++ b/Engine/Source/Editor/WorldBrowser/Private/SWorldDetails.cpp
@@ -42,6 +42,7 @@ void SWorldDetails::OnBrowseWorld(UWorld* InWorld)
 		FPropertyEditorModule& PropertyModule = FModuleManager::GetModuleChecked<FPropertyEditorModule>("PropertyEditor");
 		WorldModel->UnregisterDetailsCustomization(PropertyModule, DetailsView);
 		WorldModel->SelectionChanged.RemoveAll(this);
+		WorldModel->CollectionChanged.RemoveAll(this);
 	}
 	
 	WorldModel = nullptr;
@@ -53,6 +54,7 @@ void SWorldDetails::OnBrowseWorld(UWorld* InWorld)
 		WorldModel = WorldBrowserModule.SharedWorldModel(InWorld);
 
 		WorldModel->SelectionChanged.AddSP(this, &SWorldDetails::OnSelectionChanged);
+		WorldModel->CollectionChanged.AddSP(this, &SWorldDetails::OnCollectionChanged);
 	
 		FPropertyEditorModule& PropertyModule = FModuleManager::LoadModuleChecked<FPropertyEditorModule>("PropertyEditor");
 		FDetailsViewArgs Args(false, false, false, false, true);
@@ -176,10 +178,18 @@ void SWorldDetails::OnSelectionChanged()
 	DetailsView->SetObjects(TileProperties, true);
 }
 
+void SWorldDetails::OnCollectionChanged()
+{
+//	SubLevelsComboBox->RefreshOptions();
+}
+
 void SWorldDetails::OnSetInspectedLevel(TSharedPtr<FLevelModel> InLevelModel, ESelectInfo::Type SelectInfo)
 {
-	FLevelModelList SelectedLevels; SelectedLevels.Add(InLevelModel);
-	WorldModel->SetSelectedLevels(SelectedLevels);
+	if (InLevelModel.IsValid())
+	{
+		FLevelModelList SelectedLevels; SelectedLevels.Add(InLevelModel);
+		WorldModel->SetSelectedLevels(SelectedLevels);
+	}
 }
 
 TSharedRef<SWidget> SWorldDetails::HandleInspectedLevelComboBoxGenerateWidget(TSharedPtr<FLevelModel> InLevelModel) const
diff --git a/Engine/Source/Editor/WorldBrowser/Private/SWorldDetails.h b/Engine/Source/Editor/WorldBrowser/Private/SWorldDetails.h
--- a/Engine/Source/Editor/WorldBrowser/Private/SWorldDetails.h
+++ b/Engine/Source/Editor/WorldBrowser/Private/SWorldDetails.h
@@ -29,6 +29,9 @@ private:
 	/** Handles selection changes in data source */
 	void OnSelectionChanged();
 
+	/** Handles levels collection changes in data source */
+	void OnCollectionChanged();
+
 	/**  */
 	void OnSetInspectedLevel(TSharedPtr<FLevelModel> InLevelModel, ESelectInfo::Type SelectInfo);
 	TSharedRef<SWidget> HandleInspectedLevelComboBoxGenerateWidget(TSharedPtr<FLevelModel> InLevelModel) const;
