More fixes to "Inspect Level" selection

diff --git a/Engine/Source/Editor/WorldBrowser/Private/LevelCollectionModel.cpp b/Engine/Source/Editor/WorldBrowser/Private/LevelCollectionModel.cpp
--- a/Engine/Source/Editor/WorldBrowser/Private/LevelCollectionModel.cpp
+++ b/Engine/Source/Editor/WorldBrowser/Private/LevelCollectionModel.cpp
@@ -314,20 +314,20 @@ bool FLevelCollectionModel::IsFilterActive() const
 void FLevelCollectionModel::SetSelectedLevels(const FLevelModelList& InList)
 {
 	// Clear selection flag from currently selected levels
-	for (auto It = SelectedLevelsList.CreateIterator(); It; ++It)
+	for (auto LevelModel : SelectedLevelsList)
 	{
-		(*It)->SetLevelSelectionFlag(false);
+		LevelModel->SetLevelSelectionFlag(false);
 	}
 	
 	SelectedLevelsList.Empty(); 
 	
 	// Set selection flag to selected levels	
-	for (auto It = InList.CreateConstIterator(); It; ++It)
+	for (auto LevelModel : InList)
 	{
-		if (PassesAllFilters(*It))
+		if (LevelModel.IsValid() && PassesAllFilters(LevelModel))
 		{
-			(*It)->SetLevelSelectionFlag(true);
-			SelectedLevelsList.Add(*It);
+			LevelModel->SetLevelSelectionFlag(true);
+			SelectedLevelsList.Add(LevelModel);
 		}
 	}
 
diff --git a/Engine/Source/Editor/WorldBrowser/Private/StreamingLevels/StreamingLevelCollectionModel.cpp b/Engine/Source/Editor/WorldBrowser/Private/StreamingLevels/StreamingLevelCollectionModel.cpp
--- a/Engine/Source/Editor/WorldBrowser/Private/StreamingLevels/StreamingLevelCollectionModel.cpp
+++ b/Engine/Source/Editor/WorldBrowser/Private/StreamingLevels/StreamingLevelCollectionModel.cpp
@@ -645,7 +645,7 @@ void FStreamingLevelCollectionModel::SetStreamingLevelsClass_Executed(UClass* In
 	SetSelectedLevels(SelectedLevelsCopy);
 
 	// Force a cached level list rebuild
-	PopulateLevelsList();
+	//PopulateLevelsList();
 }
 
 //streaming volumes
