World Browser: Fixed case when "Import tiled landscape" dialog still runs import when dialog is closed

diff --git a/Engine/Source/Editor/WorldBrowser/Private/Tiles/STiledLandscapeImportDlg.cpp b/Engine/Source/Editor/WorldBrowser/Private/Tiles/STiledLandscapeImportDlg.cpp
--- a/Engine/Source/Editor/WorldBrowser/Private/Tiles/STiledLandscapeImportDlg.cpp
+++ b/Engine/Source/Editor/WorldBrowser/Private/Tiles/STiledLandscapeImportDlg.cpp
@@ -42,6 +42,7 @@ static FIntPoint ExtractTileCoordinates(FString BaseFilename)
 
 void STiledLandcapeImportDlg::Construct(const FArguments& InArgs, TSharedPtr<SWindow> InParentWindow)
 {
+	bShouldImport = false;
 	ParentWindow = InParentWindow;
 	
 	ChildSlot
@@ -357,6 +358,11 @@ TSharedRef<ITableRow> STiledLandcapeImportDlg::OnGenerateWidgetForLayerDataListV
 	
 }
 
+bool STiledLandcapeImportDlg::ShouldImport() const
+{
+	return bShouldImport;
+}
+
 const FTiledLandscapeImportSettings& STiledLandcapeImportDlg::GetImportSettings() const
 {
 	return ImportSettings;
@@ -545,14 +551,15 @@ FReply STiledLandcapeImportDlg::OnClickedImport()
 	}
 					
 	ParentWindow->RequestDestroyWindow();
+	bShouldImport = true;	
 	return FReply::Handled();	
 }
 
 FReply STiledLandcapeImportDlg::OnClickedCancel()
 {
 	ParentWindow->RequestDestroyWindow();
-	
-	ImportSettings.HeightmapFileList.Empty();
+	bShouldImport = false;	
+
 	return FReply::Handled();	
 }
 
diff --git a/Engine/Source/Editor/WorldBrowser/Private/Tiles/STiledLandscapeImportDlg.h b/Engine/Source/Editor/WorldBrowser/Private/Tiles/STiledLandscapeImportDlg.h
--- a/Engine/Source/Editor/WorldBrowser/Private/Tiles/STiledLandscapeImportDlg.h
+++ b/Engine/Source/Editor/WorldBrowser/Private/Tiles/STiledLandscapeImportDlg.h
@@ -25,6 +25,7 @@ public:
 	void Construct(const FArguments& InArgs, TSharedPtr<SWindow> InParentWindow);
 
 	/** */
+	bool ShouldImport() const;	
 	const FTiledLandscapeImportSettings& GetImportSettings() const;
 
 private:
@@ -87,6 +88,9 @@ private:
 	void UpdateLandscapeLayerList();
 
 private:
+	/** */
+	bool bShouldImport;
+	
 	/** */
 	TSharedPtr<SWindow> ParentWindow;
 		
diff --git a/Engine/Source/Editor/WorldBrowser/Private/Tiles/WorldTileCollectionModel.cpp b/Engine/Source/Editor/WorldBrowser/Private/Tiles/WorldTileCollectionModel.cpp
--- a/Engine/Source/Editor/WorldBrowser/Private/Tiles/WorldTileCollectionModel.cpp
+++ b/Engine/Source/Editor/WorldBrowser/Private/Tiles/WorldTileCollectionModel.cpp
@@ -1389,7 +1389,7 @@ void FWorldTileCollectionModel::ImportTiledLandscape_Executed()
 {
 	/** Create the window to host widget */
 	TSharedRef<SWindow> ImportWidnow = SNew(SWindow)
-											.Title(LOCTEXT("TiledLandcapeImport_DialogTitle", "Import Tiled Heightmap"))
+											.Title(LOCTEXT("TiledLandcapeImport_DialogTitle", "Import Tiled Landscape"))
 											.SizingRule( ESizingRule::Autosized )
 											.SupportsMinimize(false) 
 											.SupportsMaximize(false);
@@ -1400,11 +1400,11 @@ void FWorldTileCollectionModel::ImportTiledLandscape_Executed()
 
 	/** Show the dialog window as a modal window */
 	Editor->EditorAddModalWindow(ImportWidnow);
-
-	const FTiledLandscapeImportSettings& ImportSettings = ImportDialog->GetImportSettings();
 		
-	if (ImportSettings.HeightmapFileList.Num())
+	if (ImportDialog->ShouldImport() && ImportDialog->GetImportSettings().HeightmapFileList.Num())
 	{
+		const FTiledLandscapeImportSettings& ImportSettings = ImportDialog->GetImportSettings();
+
 		// Default path for imported landscape tiles
 		// Use tile prefix as a folder name under world root
 		FString WorldRootPath = FPackageName::LongPackageNameToFilename(GetWorld()->WorldComposition->GetWorldRoot());
