Fixed wrong selection highlight in "Inspect level" ComboBox
TTP# 340031

diff --git a/Engine/Source/Editor/WorldBrowser/Private/SWorldDetails.cpp b/Engine/Source/Editor/WorldBrowser/Private/SWorldDetails.cpp
--- a/Engine/Source/Editor/WorldBrowser/Private/SWorldDetails.cpp
+++ b/Engine/Source/Editor/WorldBrowser/Private/SWorldDetails.cpp
@@ -8,6 +8,7 @@
 #define LOCTEXT_NAMESPACE "WorldBrowser"
 
 SWorldDetails::SWorldDetails()
+	: bUpdatingSelection(false)
 {
 }
 
@@ -163,6 +164,8 @@ void SWorldDetails::OnBrowseWorld(UWorld* InWorld)
 
 void SWorldDetails::OnSelectionChanged()
 {
+	bUpdatingSelection = true;
+	
 	auto SelectedLevels = WorldModel->GetSelectedLevels();
 	TArray<UObject*> TileProperties;
 	
@@ -176,6 +179,18 @@ void SWorldDetails::OnSelectionChanged()
 	}
 
 	DetailsView->SetObjects(TileProperties, true);
+
+	if (SelectedLevels.Num() == 0 || SelectedLevels.Num() > 1)
+	{
+		// Clear ComboBox selection in case we have multiple selection
+		SubLevelsComboBox->ClearSelection();
+	}
+	else
+	{
+		SubLevelsComboBox->SetSelectedItem(SelectedLevels[0]);
+	}
+
+	bUpdatingSelection = false;
 }
 
 void SWorldDetails::OnCollectionChanged()
@@ -185,7 +200,7 @@ void SWorldDetails::OnCollectionChanged()
 
 void SWorldDetails::OnSetInspectedLevel(TSharedPtr<FLevelModel> InLevelModel, ESelectInfo::Type SelectInfo)
 {
-	if (InLevelModel.IsValid())
+	if (!bUpdatingSelection && InLevelModel.IsValid())
 	{
 		FLevelModelList SelectedLevels; SelectedLevels.Add(InLevelModel);
 		WorldModel->SetSelectedLevels(SelectedLevels);
diff --git a/Engine/Source/Editor/WorldBrowser/Private/SWorldDetails.h b/Engine/Source/Editor/WorldBrowser/Private/SWorldDetails.h
--- a/Engine/Source/Editor/WorldBrowser/Private/SWorldDetails.h
+++ b/Engine/Source/Editor/WorldBrowser/Private/SWorldDetails.h
@@ -50,4 +50,5 @@ private:
 	TSharedPtr<FLevelCollectionModel>				WorldModel;
 	TSharedPtr<IDetailsView>						DetailsView;
 	TSharedPtr<SComboBox<TSharedPtr<FLevelModel>>>	SubLevelsComboBox;
+	bool											bUpdatingSelection;	
 };
