Levels tab: When an existing level is added, it doesn't appear in the list until the window is closed and reopened
TTP #340397

diff --git a/Engine/Source/Editor/WorldBrowser/Private/StreamingLevels/StreamingLevelCollectionModel.cpp b/Engine/Source/Editor/WorldBrowser/Private/StreamingLevels/StreamingLevelCollectionModel.cpp
--- a/Engine/Source/Editor/WorldBrowser/Private/StreamingLevels/StreamingLevelCollectionModel.cpp
+++ b/Engine/Source/Editor/WorldBrowser/Private/StreamingLevels/StreamingLevelCollectionModel.cpp
@@ -499,7 +499,10 @@ void FStreamingLevelCollectionModel::AddExistingLevel(bool bRemoveInvalidSelecte
 			FLevelModelList SavedInvalidSelectedLevels = InvalidSelectedLevels;
 
 			EditorLevelUtils::AddLevelsToWorld(CurrentWorld.Get(), PackageNames, AddedLevelStreamingClass);
-
+			
+			// Force a cached level list rebuild
+			PopulateLevelsList();
+			
 			if (bRemoveInvalidSelectedLevelsAfter)
 			{
 				InvalidSelectedLevels = SavedInvalidSelectedLevels;
