Moved "Remove Selected", "Change Streaming Method" comands to a "Levels" section of context menu
TTP# 340491

diff --git a/Engine/Source/Editor/WorldBrowser/Private/LevelCollectionModel.cpp b/Engine/Source/Editor/WorldBrowser/Private/LevelCollectionModel.cpp
--- a/Engine/Source/Editor/WorldBrowser/Private/LevelCollectionModel.cpp
+++ b/Engine/Source/Editor/WorldBrowser/Private/LevelCollectionModel.cpp
@@ -652,6 +652,8 @@ void FLevelCollectionModel::BuildHierarchyMenu(FMenuBuilder& InMenuBuilder) cons
 			LOCTEXT("LockHeader", "Lock"),
 			LOCTEXT("LockSubMenu_ToolTip", "Selected Level(s) lock commands"),
 			FNewMenuDelegate::CreateSP(this, &FLevelCollectionModel::FillLockMenu ) );
+
+		ExtendHierarchyMenuSection("Levels", InMenuBuilder);
 	}
 	InMenuBuilder.EndSection();
 
@@ -661,11 +663,13 @@ void FLevelCollectionModel::BuildHierarchyMenu(FMenuBuilder& InMenuBuilder) cons
 		InMenuBuilder.AddMenuEntry( Commands.SelectAllLevels );
 		InMenuBuilder.AddMenuEntry( Commands.DeselectAllLevels );
 		InMenuBuilder.AddMenuEntry( Commands.InvertLevelSelection );
+
+		ExtendHierarchyMenuSection("LevelsSelection", InMenuBuilder);
 	}
 	InMenuBuilder.EndSection();
 	
 	// Level actors selection commands
-	InMenuBuilder.BeginSection("ActorsSelection", LOCTEXT("ActorsHeader", "Actors") );
+	InMenuBuilder.BeginSection("Actors", LOCTEXT("ActorsHeader", "Actors") );
 	{
 		InMenuBuilder.AddMenuEntry( Commands.AddsActors );
 		InMenuBuilder.AddMenuEntry( Commands.RemovesActors );
@@ -674,7 +678,10 @@ void FLevelCollectionModel::BuildHierarchyMenu(FMenuBuilder& InMenuBuilder) cons
 		{
 			InMenuBuilder.AddMenuEntry( Commands.SelectStreamingVolumes );
 		}
+
+		ExtendHierarchyMenuSection("Actors", InMenuBuilder);
 	}
+
 	InMenuBuilder.EndSection();
 }
 
diff --git a/Engine/Source/Editor/WorldBrowser/Private/LevelCollectionModel.h b/Engine/Source/Editor/WorldBrowser/Private/LevelCollectionModel.h
--- a/Engine/Source/Editor/WorldBrowser/Private/LevelCollectionModel.h
+++ b/Engine/Source/Editor/WorldBrowser/Private/LevelCollectionModel.h
@@ -123,6 +123,9 @@ public:
 	
 	/**	Builds 'hierarchy' commands menu for a selected levels */
 	virtual void BuildHierarchyMenu(FMenuBuilder& InMenuBuilder) const;
+	
+	/**	Allows to extend each section in hierarchy menu*/
+	virtual void ExtendHierarchyMenuSection(FName SectionName, FMenuBuilder& InMenuBuilder) const {};
 
 	/**	Customize 'File' section in main menu  */
 	virtual void CustomizeFileMainMenu(FMenuBuilder& InMenuBuilder) const;
diff --git a/Engine/Source/Editor/WorldBrowser/Private/StreamingLevels/StreamingLevelCollectionModel.cpp b/Engine/Source/Editor/WorldBrowser/Private/StreamingLevels/StreamingLevelCollectionModel.cpp
--- a/Engine/Source/Editor/WorldBrowser/Private/StreamingLevels/StreamingLevelCollectionModel.cpp
+++ b/Engine/Source/Editor/WorldBrowser/Private/StreamingLevels/StreamingLevelCollectionModel.cpp
@@ -242,23 +242,29 @@ void FStreamingLevelCollectionModel::BuildHierarchyMenu(FMenuBuilder& InMenuBuil
 		}
 	}
 
-	// Adds common commands
+	// Add common commands
 	FLevelCollectionModel::BuildHierarchyMenu(InMenuBuilder);
+}
 
-	InMenuBuilder.BeginSection("LevelsAddChangeStreaming");
+void FStreamingLevelCollectionModel::ExtendHierarchyMenuSection(FName SectionName, FMenuBuilder& InMenuBuilder) const
+{
+	FLevelCollectionModel::ExtendHierarchyMenuSection(SectionName, InMenuBuilder);
+	
+	const FLevelCollectionCommands& Commands = FLevelCollectionCommands::Get();
+	// Add more specific streaming levels commands to Levels section
+	if (SectionName == "Levels")
 	{
 		if (AreAnyLevelsSelected() && !(IsOneLevelSelected() && GetSelectedLevels()[0]->IsPersistent()))
 		{
 			InMenuBuilder.AddMenuEntry(Commands.World_RemoveSelectedLevels);
-			
+		
 			//
 			InMenuBuilder.AddSubMenu( 
 				LOCTEXT("LevelsChangeStreamingMethod", "Change Streaming Method"),
 				LOCTEXT("LevelsChangeStreamingMethod_Tooltip", "Changes the streaming method for the selected levels"),
-				FNewMenuDelegate::CreateRaw(this, &FStreamingLevelCollectionModel::FillSetStreamingMethodMenu ) );
+				FNewMenuDelegate::CreateRaw(this, &FStreamingLevelCollectionModel::FillSetStreamingMethodMenu ));
 		}
 	}
-	InMenuBuilder.EndSection();
 }
 
 void FStreamingLevelCollectionModel::FillSetStreamingMethodMenu(FMenuBuilder& InMenuBuilder)
diff --git a/Engine/Source/Editor/WorldBrowser/Private/StreamingLevels/StreamingLevelCollectionModel.h b/Engine/Source/Editor/WorldBrowser/Private/StreamingLevels/StreamingLevelCollectionModel.h
--- a/Engine/Source/Editor/WorldBrowser/Private/StreamingLevels/StreamingLevelCollectionModel.h
+++ b/Engine/Source/Editor/WorldBrowser/Private/StreamingLevels/StreamingLevelCollectionModel.h
@@ -31,6 +31,7 @@ public:
 	virtual void UnloadLevels(const FLevelModelList& InLevelList) override;
 	virtual TSharedPtr<FLevelDragDropOp> CreateDragDropOp() const override;
 	virtual void BuildHierarchyMenu(FMenuBuilder& InMenuBuilder) const override;
+	virtual void ExtendHierarchyMenuSection(FName SectionName, FMenuBuilder& InMenuBuilder) const override;
 	virtual void CustomizeFileMainMenu(FMenuBuilder& InMenuBuilder) const override;
 	virtual void RegisterDetailsCustomization(class FPropertyEditorModule& PropertyModule, TSharedPtr<class IDetailsView> InDetailsView)  override;
 	virtual void UnregisterDetailsCustomization(class FPropertyEditorModule& PropertyModule, TSharedPtr<class IDetailsView> InDetailsView) override;
