World Browser: Changed context menus to be more consistent between different views

diff --git a/Engine/Source/Editor/WorldBrowser/Private/LevelCollectionModel.cpp b/Engine/Source/Editor/WorldBrowser/Private/LevelCollectionModel.cpp
--- a/Engine/Source/Editor/WorldBrowser/Private/LevelCollectionModel.cpp
+++ b/Engine/Source/Editor/WorldBrowser/Private/LevelCollectionModel.cpp
@@ -630,59 +630,8 @@ bool FLevelCollectionModel::PassesAllFilters(TSharedPtr<FLevelModel> Item) const
 	return false;
 }
 
-void FLevelCollectionModel::BuildGridMenu(FMenuBuilder& InMenuBuilder) const
-{
-	CacheCanExecuteSourceControlVars();
-}
-
 void FLevelCollectionModel::BuildHierarchyMenu(FMenuBuilder& InMenuBuilder) const
 {
-	const FLevelCollectionCommands& Commands = FLevelCollectionCommands::Get();
-	
-	InMenuBuilder.BeginSection("Levels", LOCTEXT("LevelsHeader", "Levels") );
-	{
-		// Visibility commands
-		InMenuBuilder.AddSubMenu( 
-			LOCTEXT("VisibilityHeader", "Visibility"),
-			LOCTEXT("VisibilitySubMenu_ToolTip", "Selected Level(s) visibility commands"),
-			FNewMenuDelegate::CreateSP(this, &FLevelCollectionModel::FillVisibilityMenu ) );
-
-		// Lock commands
-		InMenuBuilder.AddSubMenu( 
-			LOCTEXT("LockHeader", "Lock"),
-			LOCTEXT("LockSubMenu_ToolTip", "Selected Level(s) lock commands"),
-			FNewMenuDelegate::CreateSP(this, &FLevelCollectionModel::FillLockMenu ) );
-
-		ExtendHierarchyMenuSection("Levels", InMenuBuilder);
-	}
-	InMenuBuilder.EndSection();
-
-	// Level selection commands
-	InMenuBuilder.BeginSection("LevelsSelection", LOCTEXT("SelectionHeader", "Selection") );
-	{
-		InMenuBuilder.AddMenuEntry( Commands.SelectAllLevels );
-		InMenuBuilder.AddMenuEntry( Commands.DeselectAllLevels );
-		InMenuBuilder.AddMenuEntry( Commands.InvertLevelSelection );
-
-		ExtendHierarchyMenuSection("LevelsSelection", InMenuBuilder);
-	}
-	InMenuBuilder.EndSection();
-	
-	// Level actors selection commands
-	InMenuBuilder.BeginSection("Actors", LOCTEXT("ActorsHeader", "Actors") );
-	{
-		InMenuBuilder.AddMenuEntry( Commands.AddsActors );
-		InMenuBuilder.AddMenuEntry( Commands.RemovesActors );
-
-		if (AreAnyLevelsSelected() && !(IsOneLevelSelected() && SelectedLevelsList[0]->IsPersistent()))
-		{
-			InMenuBuilder.AddMenuEntry( Commands.SelectStreamingVolumes );
-		}
-
-		ExtendHierarchyMenuSection("Actors", InMenuBuilder);
-	}
-
-	InMenuBuilder.EndSection();
 }
 
 void FLevelCollectionModel::CustomizeFileMainMenu(FMenuBuilder& InMenuBuilder) const
@@ -695,7 +644,7 @@ void FLevelCollectionModel::CustomizeFileMainMenu(FMenuBuilder& InMenuBuilder) c
 	InMenuBuilder.AddSubMenu( 
 		LOCTEXT("SourceControl", "Source Control"),
 		LOCTEXT("SourceControl_ToolTip", "Source Control Options"),
-		FNewMenuDelegate::CreateSP(this, &FLevelCollectionModel::FillSourceControlMenu));
+		FNewMenuDelegate::CreateSP(this, &FLevelCollectionModel::FillSourceControlSubMenu));
 		
 	if (AreAnyLevelsSelected())
 	{
@@ -1529,7 +1478,7 @@ void FLevelCollectionModel::ExpandSelectedItems_Executed()
 }
 
 
-void FLevelCollectionModel::FillLockMenu(FMenuBuilder& InMenuBuilder)
+void FLevelCollectionModel::FillLockSubMenu(FMenuBuilder& InMenuBuilder)
 {
 	const FLevelCollectionCommands& Commands = FLevelCollectionCommands::Get();
 
@@ -1548,7 +1497,7 @@ void FLevelCollectionModel::FillLockMenu(FMenuBuilder& InMenuBuilder)
 	}
 }
 
-void FLevelCollectionModel::FillVisibilityMenu(FMenuBuilder& InMenuBuilder)
+void FLevelCollectionModel::FillVisibilitySubMenu(FMenuBuilder& InMenuBuilder)
 {
 	const FLevelCollectionCommands& Commands = FLevelCollectionCommands::Get();
 
@@ -1559,7 +1508,7 @@ void FLevelCollectionModel::FillVisibilityMenu(FMenuBuilder& InMenuBuilder)
 	InMenuBuilder.AddMenuEntry( Commands.World_HideAllLevels );
 }
 
-void FLevelCollectionModel::FillSourceControlMenu(FMenuBuilder& InMenuBuilder)
+void FLevelCollectionModel::FillSourceControlSubMenu(FMenuBuilder& InMenuBuilder)
 {
 	const FLevelCollectionCommands& Commands = FLevelCollectionCommands::Get();
 	
diff --git a/Engine/Source/Editor/WorldBrowser/Private/LevelCollectionModel.h b/Engine/Source/Editor/WorldBrowser/Private/LevelCollectionModel.h
--- a/Engine/Source/Editor/WorldBrowser/Private/LevelCollectionModel.h
+++ b/Engine/Source/Editor/WorldBrowser/Private/LevelCollectionModel.h
@@ -117,16 +117,10 @@ public:
 	
 	/**	@return	Whether specified level passes all filters */
 	virtual bool PassesAllFilters(TSharedPtr<FLevelModel> InLevelModel) const;
-
-	/**	Builds 'minimap' commands menu for a selected levels */
-	virtual void BuildGridMenu(FMenuBuilder& InMenuBuilder) const;
 	
 	/**	Builds 'hierarchy' commands menu for a selected levels */
 	virtual void BuildHierarchyMenu(FMenuBuilder& InMenuBuilder) const;
 	
-	/**	Allows to extend each section in hierarchy menu*/
-	virtual void ExtendHierarchyMenuSection(FName SectionName, FMenuBuilder& InMenuBuilder) const {};
-
 	/**	Customize 'File' section in main menu  */
 	virtual void CustomizeFileMainMenu(FMenuBuilder& InMenuBuilder) const;
 
@@ -407,13 +401,13 @@ protected:
 	}
 	
 	/** Fills MenuBulder with Lock level related commands */
-	void FillLockMenu(class FMenuBuilder& MenuBuilder);
+	void FillLockSubMenu(class FMenuBuilder& MenuBuilder);
 	
 	/** Fills MenuBulder with level visisbility related commands */
-	void FillVisibilityMenu(class FMenuBuilder& MenuBuilder);
+	void FillVisibilitySubMenu(class FMenuBuilder& MenuBuilder);
 
 	/** Fills MenuBulder with SCC related commands */
-	void FillSourceControlMenu(class FMenuBuilder& MenuBuilder);
+	void FillSourceControlSubMenu(class FMenuBuilder& MenuBuilder);
 				
 protected:
 	/**  */
diff --git a/Engine/Source/Editor/WorldBrowser/Private/StreamingLevels/StreamingLevelCollectionModel.cpp b/Engine/Source/Editor/WorldBrowser/Private/StreamingLevels/StreamingLevelCollectionModel.cpp
--- a/Engine/Source/Editor/WorldBrowser/Private/StreamingLevels/StreamingLevelCollectionModel.cpp
+++ b/Engine/Source/Editor/WorldBrowser/Private/StreamingLevels/StreamingLevelCollectionModel.cpp
@@ -219,55 +219,82 @@ void FStreamingLevelCollectionModel::BuildHierarchyMenu(FMenuBuilder& InMenuBuil
 {
 	const FLevelCollectionCommands& Commands = FLevelCollectionCommands::Get();
 
-	if (IsOneLevelSelected())
+	// We show the "level missing" commands, when missing level is selected solely
+	if (IsOneLevelSelected() && InvalidSelectedLevels.Num() == 1)
 	{
-		// We only show the level missing level commands if only 1 invalid level and no valid levels
-		if (InvalidSelectedLevels.Num() == 1)
+		InMenuBuilder.BeginSection("MissingLevel", LOCTEXT("ViewHeaderRemove", "Missing Level") );
 		{
-			InMenuBuilder.BeginSection("MissingLevel", LOCTEXT("ViewHeaderRemove", "Missing Level") );
-			{
-				InMenuBuilder.AddMenuEntry( Commands.FixUpInvalidReference );
-				InMenuBuilder.AddMenuEntry( Commands.RemoveInvalidReference );
-			}
-			InMenuBuilder.EndSection();
-		}
-		else
-		{
-			InMenuBuilder.BeginSection("Level", FText::FromName(GetSelectedLevels()[0]->GetLongPackageName()));
-			{
-				InMenuBuilder.AddMenuEntry( Commands.World_MakeLevelCurrent );
-				InMenuBuilder.AddMenuEntry( Commands.MoveActorsToSelected );
-			}
-			InMenuBuilder.EndSection();
+			InMenuBuilder.AddMenuEntry( Commands.FixUpInvalidReference );
+			InMenuBuilder.AddMenuEntry( Commands.RemoveInvalidReference );
 		}
+		InMenuBuilder.EndSection();
 	}
 
 	// Add common commands
-	FLevelCollectionModel::BuildHierarchyMenu(InMenuBuilder);
-}
-
-void FStreamingLevelCollectionModel::ExtendHierarchyMenuSection(FName SectionName, FMenuBuilder& InMenuBuilder) const
-{
-	FLevelCollectionModel::ExtendHierarchyMenuSection(SectionName, InMenuBuilder);
-	
-	const FLevelCollectionCommands& Commands = FLevelCollectionCommands::Get();
-	// Add more specific streaming levels commands to Levels section
-	if (SectionName == "Levels")
+	InMenuBuilder.BeginSection("Levels", LOCTEXT("LevelsHeader", "Levels") );
 	{
+		// Make level current
+		if (IsOneLevelSelected())
+		{
+			InMenuBuilder.AddMenuEntry( Commands.World_MakeLevelCurrent );
+		}
+		
+		// Visibility commands
+		InMenuBuilder.AddSubMenu( 
+			LOCTEXT("VisibilityHeader", "Visibility"),
+			LOCTEXT("VisibilitySubMenu_ToolTip", "Selected Level(s) visibility commands"),
+			FNewMenuDelegate::CreateSP(this, &FStreamingLevelCollectionModel::FillVisibilitySubMenu ) );
+
+		// Lock commands
+		InMenuBuilder.AddSubMenu( 
+			LOCTEXT("LockHeader", "Lock"),
+			LOCTEXT("LockSubMenu_ToolTip", "Selected Level(s) lock commands"),
+			FNewMenuDelegate::CreateSP(this, &FStreamingLevelCollectionModel::FillLockSubMenu ) );
+		
+		// Level streaming specific commands
 		if (AreAnyLevelsSelected() && !(IsOneLevelSelected() && GetSelectedLevels()[0]->IsPersistent()))
 		{
 			InMenuBuilder.AddMenuEntry(Commands.World_RemoveSelectedLevels);
-		
 			//
 			InMenuBuilder.AddSubMenu( 
 				LOCTEXT("LevelsChangeStreamingMethod", "Change Streaming Method"),
 				LOCTEXT("LevelsChangeStreamingMethod_Tooltip", "Changes the streaming method for the selected levels"),
-				FNewMenuDelegate::CreateRaw(this, &FStreamingLevelCollectionModel::FillSetStreamingMethodMenu ));
+				FNewMenuDelegate::CreateRaw(this, &FStreamingLevelCollectionModel::FillSetStreamingMethodSubMenu ));
 		}
 	}
+	InMenuBuilder.EndSection();
+	
+
+	// Level selection commands
+	InMenuBuilder.BeginSection("LevelsSelection", LOCTEXT("SelectionHeader", "Selection") );
+	{
+		InMenuBuilder.AddMenuEntry( Commands.SelectAllLevels );
+		InMenuBuilder.AddMenuEntry( Commands.DeselectAllLevels );
+		InMenuBuilder.AddMenuEntry( Commands.InvertLevelSelection );
+	}
+	InMenuBuilder.EndSection();
+	
+	// Level actors selection commands
+	InMenuBuilder.BeginSection("Actors", LOCTEXT("ActorsHeader", "Actors") );
+	{
+		InMenuBuilder.AddMenuEntry( Commands.AddsActors );
+		InMenuBuilder.AddMenuEntry( Commands.RemovesActors );
+
+		// Move selected actors to a selected level
+		if (IsOneLevelSelected())
+		{
+			InMenuBuilder.AddMenuEntry( Commands.MoveActorsToSelected );
+		}
+
+		if (AreAnyLevelsSelected() && !(IsOneLevelSelected() && SelectedLevelsList[0]->IsPersistent()))
+		{
+			InMenuBuilder.AddMenuEntry( Commands.SelectStreamingVolumes );
+		}
+	}
+	InMenuBuilder.EndSection();
 }
 
-void FStreamingLevelCollectionModel::FillSetStreamingMethodMenu(FMenuBuilder& InMenuBuilder)
+void FStreamingLevelCollectionModel::FillSetStreamingMethodSubMenu(FMenuBuilder& InMenuBuilder)
 {
 	const FLevelCollectionCommands& Commands = FLevelCollectionCommands::Get();
 	InMenuBuilder.AddMenuEntry( Commands.SetStreamingMethod_Blueprint, NAME_None, LOCTEXT("SetStreamingMethodBlueprintOverride", "Blueprint") );
@@ -285,7 +312,7 @@ void FStreamingLevelCollectionModel::CustomizeFileMainMenu(FMenuBuilder& InMenuB
 		InMenuBuilder.AddSubMenu( 
 			LOCTEXT("LevelsStreamingMethod", "Default Streaming Method"),
 			LOCTEXT("LevelsStreamingMethod_Tooltip", "Changes the default streaming method for a new levels"),
-			FNewMenuDelegate::CreateRaw(this, &FStreamingLevelCollectionModel::FillDefaultStreamingMethodMenu ) );
+			FNewMenuDelegate::CreateRaw(this, &FStreamingLevelCollectionModel::FillDefaultStreamingMethodSubMenu ) );
 		
 		InMenuBuilder.AddMenuEntry( Commands.World_CreateEmptyLevel );
 		InMenuBuilder.AddMenuEntry( Commands.World_AddExistingLevel );
@@ -295,7 +322,7 @@ void FStreamingLevelCollectionModel::CustomizeFileMainMenu(FMenuBuilder& InMenuB
 	InMenuBuilder.EndSection();
 }
 
-void FStreamingLevelCollectionModel::FillDefaultStreamingMethodMenu(FMenuBuilder& InMenuBuilder)
+void FStreamingLevelCollectionModel::FillDefaultStreamingMethodSubMenu(FMenuBuilder& InMenuBuilder)
 {
 	const FLevelCollectionCommands& Commands = FLevelCollectionCommands::Get();
 	InMenuBuilder.AddMenuEntry( Commands.SetAddStreamingMethod_Blueprint, NAME_None, LOCTEXT("SetAddStreamingMethodBlueprintOverride", "Blueprint") );
diff --git a/Engine/Source/Editor/WorldBrowser/Private/StreamingLevels/StreamingLevelCollectionModel.h b/Engine/Source/Editor/WorldBrowser/Private/StreamingLevels/StreamingLevelCollectionModel.h
--- a/Engine/Source/Editor/WorldBrowser/Private/StreamingLevels/StreamingLevelCollectionModel.h
+++ b/Engine/Source/Editor/WorldBrowser/Private/StreamingLevels/StreamingLevelCollectionModel.h
@@ -31,7 +31,6 @@ public:
 	virtual void UnloadLevels(const FLevelModelList& InLevelList) override;
 	virtual TSharedPtr<FLevelDragDropOp> CreateDragDropOp() const override;
 	virtual void BuildHierarchyMenu(FMenuBuilder& InMenuBuilder) const override;
-	virtual void ExtendHierarchyMenuSection(FName SectionName, FMenuBuilder& InMenuBuilder) const override;
 	virtual void CustomizeFileMainMenu(FMenuBuilder& InMenuBuilder) const override;
 	virtual void RegisterDetailsCustomization(class FPropertyEditorModule& PropertyModule, TSharedPtr<class IDetailsView> InDetailsView)  override;
 	virtual void UnregisterDetailsCustomization(class FPropertyEditorModule& PropertyModule, TSharedPtr<class IDetailsView> InDetailsView) override;
@@ -114,10 +113,10 @@ private:
 	void SelectStreamingVolumes_Executed();
 
 	/**  */
-	void FillSetStreamingMethodMenu(class FMenuBuilder& MenuBuilder);
+	void FillSetStreamingMethodSubMenu(class FMenuBuilder& MenuBuilder);
 	
 	/**  */
-	void FillDefaultStreamingMethodMenu(class FMenuBuilder& MenuBuilder);
+	void FillDefaultStreamingMethodSubMenu(class FMenuBuilder& MenuBuilder);
 	
 private:
 	/** Currently selected NULL Levels */
diff --git a/Engine/Source/Editor/WorldBrowser/Private/Tiles/SWorldComposition.cpp b/Engine/Source/Editor/WorldBrowser/Private/Tiles/SWorldComposition.cpp
--- a/Engine/Source/Editor/WorldBrowser/Private/Tiles/SWorldComposition.cpp
+++ b/Engine/Source/Editor/WorldBrowser/Private/Tiles/SWorldComposition.cpp
@@ -627,7 +627,7 @@ protected:
 	
 		// Summon context menu
 		FMenuBuilder MenuBuilder(true, WorldModel->GetCommandList());
-		WorldModel->BuildGridMenu(MenuBuilder);
+		WorldModel->BuildWorldCompositionMenu(MenuBuilder);
 		TSharedPtr<SWidget> MenuWidget = MenuBuilder.MakeWidget();
 
 		FSlateApplication::Get().PushMenu(
diff --git a/Engine/Source/Editor/WorldBrowser/Private/Tiles/WorldTileCollectionModel.cpp b/Engine/Source/Editor/WorldBrowser/Private/Tiles/WorldTileCollectionModel.cpp
--- a/Engine/Source/Editor/WorldBrowser/Private/Tiles/WorldTileCollectionModel.cpp
+++ b/Engine/Source/Editor/WorldBrowser/Private/Tiles/WorldTileCollectionModel.cpp
@@ -270,157 +270,205 @@ bool FWorldTileCollectionModel::PassesAllFilters(TSharedPtr<FLevelModel> Item) c
 	return FLevelCollectionModel::PassesAllFilters(Item);
 }
 
-void FWorldTileCollectionModel::BuildGridMenu(FMenuBuilder& MenuBuilder) const
+void FWorldTileCollectionModel::BuildWorldCompositionMenu(FMenuBuilder& InMenuBuilder) const
 {
-	FLevelCollectionModel::BuildGridMenu(MenuBuilder);
-	
 	const FLevelCollectionCommands& Commands = FLevelCollectionCommands::Get();
 	
-	if (!AreAnyLevelsSelected() && GetWorld()->GetWorldSettings()->bEnableWorldOriginRebasing)
+	// No selection, option to reset world origin
+	if (!AreAnyLevelsSelected())
 	{
-		MenuBuilder.AddMenuEntry(Commands.ResetWorldOrigin);
+		if (GetWorld()->GetWorldSettings()->bEnableWorldOriginRebasing)
+		{
+			InMenuBuilder.AddMenuEntry(Commands.ResetWorldOrigin);
+		}
+		
 		return;
 	}
-		
-	MenuBuilder.BeginSection("WorldGridTemLoadUnload");
+	
+	// General Levels commands
+	InMenuBuilder.BeginSection("Levels", LOCTEXT("LevelsHeader", "Levels") );
 	{
-		MenuBuilder.AddMenuEntry(Commands.World_LoadLevel);
-		MenuBuilder.AddMenuEntry(Commands.World_UnloadLevel);
-		MenuBuilder.AddMenuEntry(Commands.World_SaveSelectedLevels);
-	}
-	MenuBuilder.EndSection();
+		// Make level current
+		if (IsOneLevelSelected())
+		{
+			InMenuBuilder.AddMenuEntry( Commands.World_MakeLevelCurrent );
+		}
 
-	MenuBuilder.BeginSection("WorldGridItemActors");
-	{
-		MenuBuilder.AddMenuEntry(Commands.AddsActors);
-		MenuBuilder.AddMenuEntry(Commands.RemovesActors);
-		MenuBuilder.AddMenuEntry(Commands.ResetLevelOrigin);
+		// Load/Unload/Save
+		InMenuBuilder.AddMenuEntry(Commands.World_LoadLevel);
+		InMenuBuilder.AddMenuEntry(Commands.World_UnloadLevel);
+		InMenuBuilder.AddMenuEntry(Commands.World_SaveSelectedLevels);
+		
+		// Visibility commands
+		InMenuBuilder.AddSubMenu( 
+			LOCTEXT("VisibilityHeader", "Visibility"),
+			LOCTEXT("VisibilitySubMenu_ToolTip", "Selected Level(s) visibility commands"),
+			FNewMenuDelegate::CreateSP(this, &FWorldTileCollectionModel::FillVisibilitySubMenu ) );
+
+		// Lock commands
+		InMenuBuilder.AddSubMenu( 
+			LOCTEXT("LockHeader", "Lock"),
+			LOCTEXT("LockSubMenu_ToolTip", "Selected Level(s) lock commands"),
+			FNewMenuDelegate::CreateSP(this, &FWorldTileCollectionModel::FillLockSubMenu ) );
 	}
-	MenuBuilder.EndSection();
+	InMenuBuilder.EndSection();
 
-	if (IsOneLevelSelected())
+	// Assign to layer
+	if (AreAnySelectedLevelsEditable())
 	{
-		MenuBuilder.BeginSection("WorldGridItemActors");
+		InMenuBuilder.BeginSection("Menu_LayersSection");
 		{
-			MenuBuilder.AddMenuEntry(Commands.MoveActorsToSelected);
+			InMenuBuilder.AddSubMenu( 
+				LOCTEXT("Layer_Assign", "Assign to Layer"),
+				FText::GetEmpty(),
+				FNewMenuDelegate::CreateSP(this, &FWorldTileCollectionModel::FillLayersSubMenu));
 		}
-		MenuBuilder.EndSection();
+		InMenuBuilder.EndSection();
+	}
 
-		if (GetWorld()->GetWorldSettings()->bEnableWorldOriginRebasing)
+	// Origin
+	InMenuBuilder.BeginSection("Menu_LevelOriginSection");
+	{
+		// Reset level position
+		InMenuBuilder.AddMenuEntry(Commands.ResetLevelOrigin);
+		
+		// Move world orign to level position
+		if (IsOneLevelSelected() && GetWorld()->GetWorldSettings()->bEnableWorldOriginRebasing)
 		{
-			MenuBuilder.BeginSection("WorldGridItemsWorldOrigin");
-			{
-				MenuBuilder.AddMenuEntry(Commands.MoveWorldOrigin);
-			}
-			MenuBuilder.EndSection();
+			InMenuBuilder.AddMenuEntry(Commands.MoveWorldOrigin);
 		}
 	}
-	else
+	InMenuBuilder.EndSection();
+
+	// Level actors selection commands
+	InMenuBuilder.BeginSection("Actors", LOCTEXT("ActorsHeader", "Actors") );
 	{
-		MenuBuilder.BeginSection("WorldGridItemMergeLevels");
+		InMenuBuilder.AddMenuEntry( Commands.AddsActors );
+		InMenuBuilder.AddMenuEntry( Commands.RemovesActors );
+
+		if (IsOneLevelSelected())
 		{
-			MenuBuilder.AddMenuEntry(Commands.World_MergeSelectedLevels);
+			InMenuBuilder.AddMenuEntry(Commands.MoveActorsToSelected);
 		}
-		MenuBuilder.EndSection();
 	}
+	InMenuBuilder.EndSection();
 
 	// Landscape specific stuff
 	const bool bCanReimportTiledLandscape = CanReimportTiledlandscape();
 	const bool bCanAddAdjacentLandscape = CanAddLandscapeProxy(FWorldTileModel::EWorldDirections::Any);
 	if (bCanReimportTiledLandscape || bCanAddAdjacentLandscape)
 	{
-		MenuBuilder.BeginSection("Menu_LandscapeSection", LOCTEXT("Menu_LandscapeSectionTitle", "Landscape"));
+		InMenuBuilder.BeginSection("Menu_LandscapeSection", LOCTEXT("Menu_LandscapeSectionTitle", "Landscape"));
 		
 		// Adjacent landscape
 		if (bCanAddAdjacentLandscape)
 		{
-			MenuBuilder.AddSubMenu( 
+			InMenuBuilder.AddSubMenu( 
 				LOCTEXT("AddLandscapeLevel", "Add Adjacent Landscape Level"),
 				FText::GetEmpty(),
-				FNewMenuDelegate::CreateSP(this, &FWorldTileCollectionModel::BuildAdjacentLandscapeMenu));
+				FNewMenuDelegate::CreateSP(this, &FWorldTileCollectionModel::FillAdjacentLandscapeSubMenu));
 		}
 
 		// Tiled landscape
 		if (bCanReimportTiledLandscape)
 		{
-			MenuBuilder.AddSubMenu( 
+			InMenuBuilder.AddSubMenu( 
 				LOCTEXT("ReimportTiledLandscape", "Reimport Tiled Landscape"),
 				FText::GetEmpty(),
-				FNewMenuDelegate::CreateSP(this, &FWorldTileCollectionModel::BuildReimportTiledLandscapeMenu));
+				FNewMenuDelegate::CreateSP(this, &FWorldTileCollectionModel::FillReimportTiledLandscapeSubMenu));
 		}
 		
-		MenuBuilder.EndSection();
-	}
-
-	if (AreAnySelectedLevelsEditable())
-	{
-		MenuBuilder.BeginSection("WorldGridItemLayers");
-		{
-			MenuBuilder.AddSubMenu( 
-				LOCTEXT("Layer_Assign", "Assign to Layer"),
-				FText::GetEmpty(),
-				FNewMenuDelegate::CreateSP(this, &FWorldTileCollectionModel::BuildLayersMenu));
-		}
-		MenuBuilder.EndSection();
+		InMenuBuilder.EndSection();
 	}
 }
 
-void FWorldTileCollectionModel::BuildHierarchyMenu(FMenuBuilder& MenuBuilder) const
+void FWorldTileCollectionModel::BuildHierarchyMenu(FMenuBuilder& InMenuBuilder) const
 {
 	const FLevelCollectionCommands& Commands = FLevelCollectionCommands::Get();
 		
-	if (IsOneLevelSelected() && SelectedLevelsList[0]->IsLoaded())
+	// Add common commands
+	InMenuBuilder.BeginSection("Levels", LOCTEXT("LevelsHeader", "Levels") );
 	{
-		MenuBuilder.BeginSection("Level", FText::FromName(SelectedLevelsList[0]->GetLongPackageName()));
+		// Make level current
+		if (IsOneLevelSelected())
 		{
-			MenuBuilder.AddMenuEntry( Commands.World_MakeLevelCurrent );
-			MenuBuilder.AddMenuEntry( Commands.MoveActorsToSelected );
+			InMenuBuilder.AddMenuEntry( Commands.World_MakeLevelCurrent );
 		}
-		MenuBuilder.EndSection();
-	}
+
+		// Load/Unload/Save
+		InMenuBuilder.AddMenuEntry(Commands.World_LoadLevel);
+		InMenuBuilder.AddMenuEntry(Commands.World_UnloadLevel);
+		InMenuBuilder.AddMenuEntry(Commands.World_SaveSelectedLevels);
+				
+		// Visibility commands
+		InMenuBuilder.AddSubMenu( 
+			LOCTEXT("VisibilityHeader", "Visibility"),
+			LOCTEXT("VisibilitySubMenu_ToolTip", "Selected Level(s) visibility commands"),
+			FNewMenuDelegate::CreateSP(this, &FWorldTileCollectionModel::FillVisibilitySubMenu ) );
+
+		// Lock commands
+		InMenuBuilder.AddSubMenu( 
+			LOCTEXT("LockHeader", "Lock"),
+			LOCTEXT("LockSubMenu_ToolTip", "Selected Level(s) lock commands"),
+			FNewMenuDelegate::CreateSP(this, &FWorldTileCollectionModel::FillLockSubMenu ) );
 	
-	if (AreAnyLevelsSelected())
+	}
+	InMenuBuilder.EndSection();
+
+	// Assign to layer
+	if (AreAnySelectedLevelsEditable())
 	{
-		MenuBuilder.BeginSection("Menu_LoadUnload");
+		InMenuBuilder.BeginSection("Menu_LayersSection");
 		{
-			MenuBuilder.AddMenuEntry(Commands.World_LoadLevel);
-			MenuBuilder.AddMenuEntry(Commands.World_UnloadLevel);
+			InMenuBuilder.AddSubMenu( 
+				LOCTEXT("WorldLayers", "Assign to Layer"),
+				FText::GetEmpty(),
+				FNewMenuDelegate::CreateSP(this, &FWorldTileCollectionModel::FillLayersSubMenu)
+			);
 		}
-		MenuBuilder.EndSection();
+		InMenuBuilder.EndSection();
 	}
 
-	// default commands
-	FLevelCollectionModel::BuildHierarchyMenu(MenuBuilder);
-	
+	// Hierarchy management
 	if (AreAnyLevelsSelected())
 	{
-		MenuBuilder.BeginSection("Menu_Hierarchy");
+		InMenuBuilder.BeginSection("Menu_HierarchySection");
 		{
-			MenuBuilder.AddMenuEntry(Commands.ExpandSelectedItems);
-			MenuBuilder.AddMenuEntry(Commands.ClearParentLink);
+			InMenuBuilder.AddMenuEntry(Commands.ExpandSelectedItems);
+			InMenuBuilder.AddMenuEntry(Commands.ClearParentLink);
 		}
-		MenuBuilder.EndSection();
+		InMenuBuilder.EndSection();
 	}
-	
-	if (AreAnySelectedLevelsEditable())
+		
+	// Level selection commands
+	InMenuBuilder.BeginSection("LevelsSelection", LOCTEXT("SelectionHeader", "Selection") );
 	{
-		MenuBuilder.BeginSection("Menu_Layers");
+		InMenuBuilder.AddMenuEntry( Commands.SelectAllLevels );
+		InMenuBuilder.AddMenuEntry( Commands.DeselectAllLevels );
+		InMenuBuilder.AddMenuEntry( Commands.InvertLevelSelection );
+	}
+	InMenuBuilder.EndSection();
+
+	// Level actors selection commands
+	InMenuBuilder.BeginSection("Actors", LOCTEXT("ActorsHeader", "Actors") );
+	{
+		InMenuBuilder.AddMenuEntry( Commands.AddsActors );
+		InMenuBuilder.AddMenuEntry( Commands.RemovesActors );
+
+		// Move selected actors to a selected level
+		if (IsOneLevelSelected())
 		{
-			MenuBuilder.AddSubMenu( 
-				LOCTEXT("WorldLayers", "Assign to Layer"),
-				FText::GetEmpty(),
-				FNewMenuDelegate::CreateSP(this, &FWorldTileCollectionModel::BuildLayersMenu)
-			);
+			InMenuBuilder.AddMenuEntry( Commands.MoveActorsToSelected );
 		}
-		MenuBuilder.EndSection();
 	}
+	InMenuBuilder.EndSection();
 }
 
-void FWorldTileCollectionModel::BuildLayersMenu(FMenuBuilder& MenuBuilder) const
+void FWorldTileCollectionModel::FillLayersSubMenu(FMenuBuilder& InMenuBuilder) const
 {
 	for (auto It = AllLayers.CreateConstIterator(); It; ++It)
 	{
-		MenuBuilder.AddMenuEntry(
+		InMenuBuilder.AddMenuEntry(
 			FText::FromString((*It).Name), 
 			FText(), FSlateIcon(),
 			FUIAction(FExecuteAction::CreateSP(
@@ -431,17 +479,17 @@ void FWorldTileCollectionModel::BuildLayersMenu(FMenuBuilder& MenuBuilder) const
 	}
 }
 
-void FWorldTileCollectionModel::BuildAdjacentLandscapeMenu(FMenuBuilder& MenuBuilder) const
+void FWorldTileCollectionModel::FillAdjacentLandscapeSubMenu(FMenuBuilder& InMenuBuilder) const
 {
 	const FLevelCollectionCommands& Commands = FLevelCollectionCommands::Get();
 
-	MenuBuilder.AddMenuEntry(Commands.AddLandscapeLevelXNegative, NAME_None, TAttribute<FText>(), TAttribute<FText>(), FSlateIcon(FEditorStyle::GetStyleSetName(), "WorldBrowser.DirectionXNegative"));
-	MenuBuilder.AddMenuEntry(Commands.AddLandscapeLevelXPositive, NAME_None, TAttribute<FText>(), TAttribute<FText>(), FSlateIcon(FEditorStyle::GetStyleSetName(), "WorldBrowser.DirectionXPositive"));
-	MenuBuilder.AddMenuEntry(Commands.AddLandscapeLevelYNegative, NAME_None, TAttribute<FText>(), TAttribute<FText>(), FSlateIcon(FEditorStyle::GetStyleSetName(), "WorldBrowser.DirectionYNegative"));
-	MenuBuilder.AddMenuEntry(Commands.AddLandscapeLevelYPositive, NAME_None, TAttribute<FText>(), TAttribute<FText>(), FSlateIcon(FEditorStyle::GetStyleSetName(), "WorldBrowser.DirectionYPositive"));
+	InMenuBuilder.AddMenuEntry(Commands.AddLandscapeLevelXNegative, NAME_None, TAttribute<FText>(), TAttribute<FText>(), FSlateIcon(FEditorStyle::GetStyleSetName(), "WorldBrowser.DirectionXNegative"));
+	InMenuBuilder.AddMenuEntry(Commands.AddLandscapeLevelXPositive, NAME_None, TAttribute<FText>(), TAttribute<FText>(), FSlateIcon(FEditorStyle::GetStyleSetName(), "WorldBrowser.DirectionXPositive"));
+	InMenuBuilder.AddMenuEntry(Commands.AddLandscapeLevelYNegative, NAME_None, TAttribute<FText>(), TAttribute<FText>(), FSlateIcon(FEditorStyle::GetStyleSetName(), "WorldBrowser.DirectionYNegative"));
+	InMenuBuilder.AddMenuEntry(Commands.AddLandscapeLevelYPositive, NAME_None, TAttribute<FText>(), TAttribute<FText>(), FSlateIcon(FEditorStyle::GetStyleSetName(), "WorldBrowser.DirectionYPositive"));
 }
 
-void FWorldTileCollectionModel::BuildReimportTiledLandscapeMenu(FMenuBuilder& InMenuBuilder) const
+void FWorldTileCollectionModel::FillReimportTiledLandscapeSubMenu(FMenuBuilder& InMenuBuilder) const
 {
 	const FLevelCollectionCommands& Commands = FLevelCollectionCommands::Get();
 
@@ -459,11 +507,11 @@ void FWorldTileCollectionModel::BuildReimportTiledLandscapeMenu(FMenuBuilder& In
 	InMenuBuilder.AddSubMenu( 
 				LOCTEXT("Menu_WeightmapsTitle", "Weightmaps"),
 				FText::GetEmpty(),
-				FNewMenuDelegate::CreateSP(this, &FWorldTileCollectionModel::BuildWeightmapsMenu)
+				FNewMenuDelegate::CreateSP(this, &FWorldTileCollectionModel::FillWeightmapsSubMenu)
 			);
 }
 
-void FWorldTileCollectionModel::BuildWeightmapsMenu(FMenuBuilder& InMenuBuilder) const
+void FWorldTileCollectionModel::FillWeightmapsSubMenu(FMenuBuilder& InMenuBuilder) const
 {
 	// Add "All Weighmaps" menu entry
 	InMenuBuilder.AddMenuEntry(
diff --git a/Engine/Source/Editor/WorldBrowser/Private/Tiles/WorldTileCollectionModel.h b/Engine/Source/Editor/WorldBrowser/Private/Tiles/WorldTileCollectionModel.h
--- a/Engine/Source/Editor/WorldBrowser/Private/Tiles/WorldTileCollectionModel.h
+++ b/Engine/Source/Editor/WorldBrowser/Private/Tiles/WorldTileCollectionModel.h
@@ -31,7 +31,6 @@ public:
 	virtual FVector2D SnapTranslationDelta(const FLevelModelList& InList, FVector2D InTranslationDelta, bool bBoundsSnapping, float InSnappingValue) override;
 	virtual TSharedPtr<FLevelDragDropOp> CreateDragDropOp() const override;
 	virtual bool PassesAllFilters(TSharedPtr<FLevelModel> InLevelModel) const override;
-	virtual void BuildGridMenu(FMenuBuilder& InMenuBuilder) const override;
 	virtual void BuildHierarchyMenu(FMenuBuilder& InMenuBuilder) const override;
 	virtual void CustomizeFileMainMenu(FMenuBuilder& InMenuBuilder) const override;
 	virtual FMatrix GetObserverViewMatrix() const override;
@@ -138,6 +137,11 @@ public:
 	 */
 	void Focus(FBox InArea, FocusStrategy InStaragegy);
 
+	/** 
+	 *  Builds context menu for a world composition
+	 */
+	void BuildWorldCompositionMenu(FMenuBuilder& InMenuBuilder) const;
+
 private:
 	FWorldTileCollectionModel( const TWeakObjectPtr< UEditorEngine >& InEditor );
 	
@@ -181,17 +185,17 @@ private:
 	 */
 	bool AddLevelToTheWorld(const TSharedPtr<FWorldTileModel>& InLevel);
 	
-	/** Builds Layers menu */
-	void BuildLayersMenu(FMenuBuilder& InMenuBuilder) const;
+	/** Fills Layers sub-menu */
+	void FillLayersSubMenu(FMenuBuilder& InMenuBuilder) const;
 	
-	/** Builds adjacent landscape menu */
-	void BuildAdjacentLandscapeMenu(FMenuBuilder& InMenuBuilder) const;
+	/** Fills adjacent landscape sub-menu */
+	void FillAdjacentLandscapeSubMenu(FMenuBuilder& InMenuBuilder) const;
 
-	/** Builds reimport tiled landscape menu */
-	void BuildReimportTiledLandscapeMenu(FMenuBuilder& InMenuBuilder) const;
+	/** Fills reimport tiled landscape sub-menu */
+	void FillReimportTiledLandscapeSubMenu(FMenuBuilder& InMenuBuilder) const;
 	
-	/** Builds reimport weightmaps menu */
-	void BuildWeightmapsMenu(FMenuBuilder& InMenuBuilder) const;
+	/** Fills reimport weightmaps sub_menu */
+	void FillWeightmapsSubMenu(FMenuBuilder& InMenuBuilder) const;
 	
 private:
 	/** Creates a new empty Level; prompts for level save location */
