Fixed: After changing the streaming method of a level, it is not updated in the sub menu correctly TTP# 341080

diff --git a/Engine/Source/Editor/WorldBrowser/Private/StreamingLevels/StreamingLevelCollectionModel.cpp b/Engine/Source/Editor/WorldBrowser/Private/StreamingLevels/StreamingLevelCollectionModel.cpp
--- a/Engine/Source/Editor/WorldBrowser/Private/StreamingLevels/StreamingLevelCollectionModel.cpp
+++ b/Engine/Source/Editor/WorldBrowser/Private/StreamingLevels/StreamingLevelCollectionModel.cpp
@@ -679,7 +679,7 @@ void FStreamingLevelCollectionModel::SetStreamingLevelsClass_Executed(UClass* In
 	SetSelectedLevels(SelectedLevelsCopy);
 
 	// Force a cached level list rebuild
-	//PopulateLevelsList();
+	PopulateLevelsList();
 }
 
 //streaming volumes
