Fixed: A new blank map doesnÆt have a Persistent Level
TTP# 341247

diff --git a/Engine/Source/Editor/WorldBrowser/Private/LevelCollectionModel.cpp b/Engine/Source/Editor/WorldBrowser/Private/LevelCollectionModel.cpp
--- a/Engine/Source/Editor/WorldBrowser/Private/LevelCollectionModel.cpp
+++ b/Engine/Source/Editor/WorldBrowser/Private/LevelCollectionModel.cpp
@@ -10,7 +10,7 @@
 
 #define LOCTEXT_NAMESPACE "WorldBrowser"
 
-FLevelCollectionModel::FLevelCollectionModel(const TWeakObjectPtr<UEditorEngine>& InEditor)
+FLevelCollectionModel::FLevelCollectionModel(UEditorEngine* InEditor)
 	: Editor(InEditor)
 	, bRequestedUpdateAllLevels(false)
 	, bRequestedRedrawAllLevels(false)
@@ -39,11 +39,11 @@ FLevelCollectionModel::~FLevelCollectionModel()
 	Editor->OnLevelActorDeleted().RemoveAll(this);
 }
 
-void FLevelCollectionModel::Initialize()
+void FLevelCollectionModel::Initialize(UWorld* InWorld)
 {
 	LoadSettings();
 	
-	CurrentWorld = Editor->GetEditorWorldContext().World();
+	CurrentWorld = InWorld;
 
 	Filters->OnChanged().AddSP(this, &FLevelCollectionModel::OnFilterChanged);
 	FWorldDelegates::LevelAddedToWorld.AddSP(this, &FLevelCollectionModel::OnLevelAddedToWorld);
diff --git a/Engine/Source/Editor/WorldBrowser/Private/LevelCollectionModel.h b/Engine/Source/Editor/WorldBrowser/Private/LevelCollectionModel.h
--- a/Engine/Source/Editor/WorldBrowser/Private/LevelCollectionModel.h
+++ b/Engine/Source/Editor/WorldBrowser/Private/LevelCollectionModel.h
@@ -22,7 +22,7 @@ public:
 
 
 public:
-	FLevelCollectionModel(const TWeakObjectPtr<UEditorEngine>& InEditor);
+	FLevelCollectionModel(UEditorEngine* InEditor);
 	virtual ~FLevelCollectionModel();
 
 	/** FTickableEditorObject interface */
@@ -411,7 +411,7 @@ protected:
 				
 protected:
 	/**  */
-	virtual void Initialize();
+	virtual void Initialize(UWorld* InWorld);
 	
 	/**  */
 	virtual void BindCommands();
diff --git a/Engine/Source/Editor/WorldBrowser/Private/StreamingLevels/StreamingLevelCollectionModel.cpp b/Engine/Source/Editor/WorldBrowser/Private/StreamingLevels/StreamingLevelCollectionModel.cpp
--- a/Engine/Source/Editor/WorldBrowser/Private/StreamingLevels/StreamingLevelCollectionModel.cpp
+++ b/Engine/Source/Editor/WorldBrowser/Private/StreamingLevels/StreamingLevelCollectionModel.cpp
@@ -13,7 +13,7 @@
 
 #define LOCTEXT_NAMESPACE "WorldBrowser"
 
-FStreamingLevelCollectionModel::FStreamingLevelCollectionModel( const TWeakObjectPtr< UEditorEngine >& InEditor )
+FStreamingLevelCollectionModel::FStreamingLevelCollectionModel(UEditorEngine* InEditor)
 	: FLevelCollectionModel(InEditor)
 	, AddedLevelStreamingClass(ULevelStreamingKismet::StaticClass())
 {
@@ -24,12 +24,12 @@ FStreamingLevelCollectionModel::~FStreamingLevelCollectionModel()
 	Editor->UnregisterForUndo( this );
 }
 
-void FStreamingLevelCollectionModel::Initialize()
+void FStreamingLevelCollectionModel::Initialize(UWorld* InWorld)
 {
 	BindCommands();	
 	Editor->RegisterForUndo( this );
 	
-	FLevelCollectionModel::Initialize();
+	FLevelCollectionModel::Initialize(InWorld);
 }
 
 void FStreamingLevelCollectionModel::OnLevelsCollectionChanged()
diff --git a/Engine/Source/Editor/WorldBrowser/Private/StreamingLevels/StreamingLevelCollectionModel.h b/Engine/Source/Editor/WorldBrowser/Private/StreamingLevels/StreamingLevelCollectionModel.h
--- a/Engine/Source/Editor/WorldBrowser/Private/StreamingLevels/StreamingLevelCollectionModel.h
+++ b/Engine/Source/Editor/WorldBrowser/Private/StreamingLevels/StreamingLevelCollectionModel.h
@@ -19,10 +19,10 @@ public:
 	 *
 	 *	@param	InEditor		The UEditorEngine to use
 	 */
-	static TSharedRef<FStreamingLevelCollectionModel> Create(const TWeakObjectPtr<UEditorEngine>& InEditor)
+	static TSharedRef<FStreamingLevelCollectionModel> Create(UEditorEngine* InEditor, UWorld* InWorld)
 	{
 		TSharedRef<FStreamingLevelCollectionModel> LevelCollectionModel(new FStreamingLevelCollectionModel(InEditor));
-		LevelCollectionModel->Initialize();
+		LevelCollectionModel->Initialize(InWorld);
 		return LevelCollectionModel;
 	}
 
@@ -36,7 +36,7 @@ public:
 	virtual void UnregisterDetailsCustomization(class FPropertyEditorModule& PropertyModule, TSharedPtr<class IDetailsView> InDetailsView) override;
 
 private:
-	virtual void Initialize() override;
+	virtual void Initialize(UWorld* InWorld) override;
 	virtual void BindCommands() override;
 	virtual void OnLevelsCollectionChanged() override;
 	virtual void OnLevelsSelectionChanged() override;
@@ -59,7 +59,7 @@ private:
 	 *	@param	InWorldLevels	The Level management logic object
 	 *	@param	InEditor		The UEditorEngine to use
 	 */
-	FStreamingLevelCollectionModel( const TWeakObjectPtr< UEditorEngine >& InEditor );
+	FStreamingLevelCollectionModel(UEditorEngine* InEditor);
 	
 	/** Refreshes the sort index on all viewmodels that contain ULevels */
 	void RefreshSortIndexes();
diff --git a/Engine/Source/Editor/WorldBrowser/Private/Tiles/WorldTileCollectionModel.cpp b/Engine/Source/Editor/WorldBrowser/Private/Tiles/WorldTileCollectionModel.cpp
--- a/Engine/Source/Editor/WorldBrowser/Private/Tiles/WorldTileCollectionModel.cpp
+++ b/Engine/Source/Editor/WorldBrowser/Private/Tiles/WorldTileCollectionModel.cpp
@@ -23,7 +23,7 @@
 static const FName HeightmapLayerName = FName("__Heightmap__");
 
 
-FWorldTileCollectionModel::FWorldTileCollectionModel(const TWeakObjectPtr<UEditorEngine>& InEditor)
+FWorldTileCollectionModel::FWorldTileCollectionModel(UEditorEngine* InEditor)
 	: FLevelCollectionModel(InEditor)
 	, bIsSavingLevel(false)
 	, bMeshProxyAvailable(false)
@@ -45,7 +45,7 @@ FWorldTileCollectionModel::~FWorldTileCollectionModel()
 	FEditorDelegates::NewCurrentLevel.RemoveAll(this);
 }
 
-void FWorldTileCollectionModel::Initialize()
+void FWorldTileCollectionModel::Initialize(UWorld* InWorld)
 {
 	// Uncategorized layer, always exist
 	FWorldTileLayer Layer;
@@ -60,7 +60,7 @@ void FWorldTileCollectionModel::Initialize()
 	FEditorDelegates::NewCurrentLevel.AddSP(this, &FWorldTileCollectionModel::OnNewCurrentLevel);
 	BindCommands();
 			
-	FLevelCollectionModel::Initialize();
+	FLevelCollectionModel::Initialize(InWorld);
 	
 	// Check whehter Editor has support for generating mesh proxies	
 	IMeshUtilities& MeshUtilities = FModuleManager::Get().LoadModuleChecked<IMeshUtilities>("MeshUtilities");
diff --git a/Engine/Source/Editor/WorldBrowser/Private/Tiles/WorldTileCollectionModel.h b/Engine/Source/Editor/WorldBrowser/Private/Tiles/WorldTileCollectionModel.h
--- a/Engine/Source/Editor/WorldBrowser/Private/Tiles/WorldTileCollectionModel.h
+++ b/Engine/Source/Editor/WorldBrowser/Private/Tiles/WorldTileCollectionModel.h
@@ -18,10 +18,10 @@ public:
 	 *
 	 *	@param	InEditor		The UEditorEngine to use
 	 */
-	static TSharedRef<FWorldTileCollectionModel> Create(const TWeakObjectPtr<UEditorEngine>& InEditor)
+	static TSharedRef<FWorldTileCollectionModel> Create(UEditorEngine* InEditor, UWorld* InWorld)
 	{
 		TSharedRef<FWorldTileCollectionModel> LevelCollectionModel(new FWorldTileCollectionModel(InEditor));
-		LevelCollectionModel->Initialize();
+		LevelCollectionModel->Initialize(InWorld);
 		return LevelCollectionModel;
 	}
 public:
@@ -46,7 +46,7 @@ private:
 	/** FTickableEditorObject interface end */
 
 	/** FLevelCollection interface */
-	virtual void Initialize() override;
+	virtual void Initialize(UWorld* InWorld) override;
 	virtual void BindCommands() override;
 	virtual void OnLevelsCollectionChanged() override;
 	virtual void OnLevelsSelectionChanged() override;
@@ -143,7 +143,7 @@ public:
 	void BuildWorldCompositionMenu(FMenuBuilder& InMenuBuilder) const;
 
 private:
-	FWorldTileCollectionModel( const TWeakObjectPtr< UEditorEngine >& InEditor );
+	FWorldTileCollectionModel(UEditorEngine* InEditor);
 	
 	/** Setups parent->child links between tiles */
 	void SetupParentChildLinks();
diff --git a/Engine/Source/Editor/WorldBrowser/Private/WorldBrowserModule.cpp b/Engine/Source/Editor/WorldBrowser/Private/WorldBrowserModule.cpp
--- a/Engine/Source/Editor/WorldBrowser/Private/WorldBrowserModule.cpp
+++ b/Engine/Source/Editor/WorldBrowser/Private/WorldBrowserModule.cpp
@@ -110,11 +110,11 @@ TSharedPtr<FLevelCollectionModel> FWorldBrowserModule::SharedWorldModel(UWorld*
 		{
 			if (InWorld->WorldComposition)
 			{
-				SharedWorldModel = FWorldTileCollectionModel::Create(GEditor);
+				SharedWorldModel = FWorldTileCollectionModel::Create(GEditor, InWorld);
 			}
 			else
 			{
-				SharedWorldModel = FStreamingLevelCollectionModel::Create(GEditor);
+				SharedWorldModel = FStreamingLevelCollectionModel::Create(GEditor, InWorld);
 			}
 		}
 
