Added levels reordering in a Levels browser (drag and drop)

diff --git a/Engine/Source/Editor/WorldBrowser/Private/LevelModel.h b/Engine/Source/Editor/WorldBrowser/Private/LevelModel.h
--- a/Engine/Source/Editor/WorldBrowser/Private/LevelModel.h
+++ b/Engine/Source/Editor/WorldBrowser/Private/LevelModel.h
@@ -251,10 +251,7 @@ protected:
 
 	/** Reference to owning collection model */
 	FLevelCollectionModel&				LevelCollectionModel;
-
-	/** Owning world  */
-	const TWeakObjectPtr<UWorld>		CurrentWorld;
-
+		
 	/** The UEditorEngine to use */
 	const TWeakObjectPtr<UEditorEngine> Editor;
 
diff --git a/Engine/Source/Editor/WorldBrowser/Private/StreamingLevels/StreamingLevelCollectionModel.cpp b/Engine/Source/Editor/WorldBrowser/Private/StreamingLevels/StreamingLevelCollectionModel.cpp
--- a/Engine/Source/Editor/WorldBrowser/Private/StreamingLevels/StreamingLevelCollectionModel.cpp
+++ b/Engine/Source/Editor/WorldBrowser/Private/StreamingLevels/StreamingLevelCollectionModel.cpp
@@ -64,7 +64,6 @@ void FStreamingLevelCollectionModel::OnLevelsCollectionChanged()
 		}
 	}
 
-	RefreshSortIndexes();
 	FLevelCollectionModel::OnLevelsCollectionChanged();
 }
 
@@ -345,97 +344,11 @@ void FStreamingLevelCollectionModel::UnregisterDetailsCustomization(FPropertyEdi
 	InDetailsView->UnregisterInstancedCustomPropertyLayout(ULevelStreaming::StaticClass());
 }
 
-bool FStreamingLevelCollectionModel::CanShiftSelection()
-{
-	if (!IsOneLevelSelected())
-	{
-		return false;
-	}
-
-	for (int32 i = 0; i < SelectedLevelsList.Num(); ++i)
-	{
-		if (SelectedLevelsList[i]->IsLocked() || SelectedLevelsList[i]->IsPersistent())
-		{
-			return false;
-		}
-	}
-
-	return true;
-}
-
-void FStreamingLevelCollectionModel::ShiftSelection( bool bUp )
-{
-	if (!CanShiftSelection())
-	{
-		return;
-	}
-
-	if (!IsSelectedLevelEditable())
-	{
-		FMessageDialog::Open( EAppMsgType::Ok, LOCTEXT("ShiftLevelLocked", "Shift Level: The requested operation could not be completed because the level is locked or not loaded.") );
-		return;
-	}
-	
-	TSharedPtr<FStreamingLevelModel> TargetModel = StaticCastSharedPtr<FStreamingLevelModel>(SelectedLevelsList[0]);
-	ULevelStreaming* InLevelStreaming = TargetModel->GetLevelStreaming().Get();
-	
-	int32 FoundLevelIndex =  CurrentWorld->StreamingLevels.Find(InLevelStreaming);
-	int32 PrevFoundLevelIndex = FoundLevelIndex - 1;
-	int32 PostFoundLevelIndex = FoundLevelIndex + 1;
-
-	// If we found the level . . .
-	if ( FoundLevelIndex != INDEX_NONE )
-	{
-		// Check if we found a destination index to swap it to.
-		const int32 DestIndex = bUp ? PrevFoundLevelIndex : PostFoundLevelIndex;
-		if (CurrentWorld->StreamingLevels.IsValidIndex(DestIndex))
-		{
-			// Swap the level into position.
-			const FScopedTransaction Transaction( NSLOCTEXT("UnrealEd", "ShiftLevelInLevelBrowser", "Shift level in Level Browser") );
-			CurrentWorld->Modify();
-			CurrentWorld->StreamingLevels.Swap( FoundLevelIndex, DestIndex );
-			CurrentWorld->MarkPackageDirty();
-		}
-	}
-
-	RefreshSortIndexes();
-}
-
-
 const FLevelModelList& FStreamingLevelCollectionModel::GetInvalidSelectedLevels() const 
 { 
 	return InvalidSelectedLevels;
 }
 
-void FStreamingLevelCollectionModel::RefreshSortIndexes()
-{
-	//for( auto LevelIt = AllLevelsList.CreateIterator(); LevelIt; ++LevelIt )
-	//{
-	//	(*LevelIt)->RefreshStreamingLevelIndex();
-	//}
-
-	//OnFilterChanged();
-}
-
-void FStreamingLevelCollectionModel::SortFilteredLevels()
-{
-	struct FCompareLevels
-	{
-		FORCEINLINE bool operator()(const TSharedPtr<FLevelModel>& InLhs, 
-									const TSharedPtr<FLevelModel>& InRhs) const 
-		{ 
-			TSharedPtr<FStreamingLevelModel> Lhs = StaticCastSharedPtr<FStreamingLevelModel>(InLhs);
-			TSharedPtr<FStreamingLevelModel> Rhs = StaticCastSharedPtr<FStreamingLevelModel>(InRhs);
-				
-			// Sort by a user-defined order.
-			return (Lhs->GetStreamingLevelIndex() < Rhs->GetStreamingLevelIndex());
-		}
-	};
-
-	FilteredLevelsList.Sort(FCompareLevels());
-}
-
-
 //levels
 void FStreamingLevelCollectionModel::CreateEmptyLevel_Executed()
 {
diff --git a/Engine/Source/Editor/WorldBrowser/Private/StreamingLevels/StreamingLevelCollectionModel.h b/Engine/Source/Editor/WorldBrowser/Private/StreamingLevels/StreamingLevelCollectionModel.h
--- a/Engine/Source/Editor/WorldBrowser/Private/StreamingLevels/StreamingLevelCollectionModel.h
+++ b/Engine/Source/Editor/WorldBrowser/Private/StreamingLevels/StreamingLevelCollectionModel.h
@@ -43,12 +43,6 @@ private:
 	/** FLevelCollection interface end */
 	
 public:
-	/** @return Whether the current selection can be shifted */
-	bool CanShiftSelection();
-
-	/** Moves the level selection up or down in the list; used for re-ordering */
-	void ShiftSelection( bool bUp );
-
 	/** @return Any selected ULevel objects in the LevelsView that are NULL */
 	const FLevelModelList& GetInvalidSelectedLevels() const;
 
@@ -61,12 +55,6 @@ private:
 	 */
 	FStreamingLevelCollectionModel(UEditorEngine* InEditor);
 	
-	/** Refreshes the sort index on all viewmodels that contain ULevels */
-	void RefreshSortIndexes();
-	
-	/**	Sorts the filtered Levels list */
-	void SortFilteredLevels();
-
 	// Begin FEditorUndoClient Interface
 	virtual void PostUndo(bool bSuccess) override { UpdateAllLevels(); }
 	virtual void PostRedo(bool bSuccess) override { PostUndo(bSuccess); }
diff --git a/Engine/Source/Editor/WorldBrowser/Private/StreamingLevels/StreamingLevelModel.cpp b/Engine/Source/Editor/WorldBrowser/Private/StreamingLevels/StreamingLevelModel.cpp
--- a/Engine/Source/Editor/WorldBrowser/Private/StreamingLevels/StreamingLevelModel.cpp
+++ b/Engine/Source/Editor/WorldBrowser/Private/StreamingLevels/StreamingLevelModel.cpp
@@ -13,7 +13,6 @@ FStreamingLevelModel::FStreamingLevelModel(const TWeakObjectPtr<UEditorEngine>&
 
 	: FLevelModel(InWorldData, InEditor)
 	, LevelStreaming(InLevelStreaming)
-	, StreamingLevelIndex(INDEX_NONE)
 	, bHasValidPackageName(false)
 {
 	Editor->RegisterForUndo( this );
@@ -71,27 +70,52 @@ void FStreamingLevelModel::Update()
 	FLevelModel::Update();
 }
 
-void FStreamingLevelModel::RefreshStreamingLevelIndex()
+void FStreamingLevelModel::OnDrop(const TSharedPtr<FLevelDragDropOp>& Op)
 {
-	ULevelStreaming* StreamingLevel = LevelStreaming.Get();
+	TArray<ULevelStreaming*> DropStreamingLevels;
 	
-	if ( StreamingLevel )
+	for (auto It = Op->StreamingLevelsToDrop.CreateConstIterator(); It; ++It)
 	{
-		int32 Idx = CurrentWorld->StreamingLevels.Find(StreamingLevel);
-		if (Idx != INDEX_NONE)
+		if ((*It).IsValid())
 		{
-			StreamingLevelIndex = Idx;
+			DropStreamingLevels.AddUnique((*It).Get());
+		}
+	}	
+	
+	// Prevent dropping items on itself
+	if (DropStreamingLevels.Num() && DropStreamingLevels.Find(LevelStreaming.Get()) == INDEX_NONE)
+	{
+		UWorld* CurrentWorld = LevelCollectionModel.GetWorld();
+		auto& WorldStreamingLevels = CurrentWorld->StreamingLevels;
+		// Remove streaming level objects from a world streaming levels list
+		for (auto It : DropStreamingLevels)
+		{
+			WorldStreamingLevels.Remove(It);
+		}
+		
+		// Find a new place where to insert the in a world streaming levels list
+		// Right after the current level, or at start of the list in case if this is persistent level
+		int32 InsertIndex = WorldStreamingLevels.Find(LevelStreaming.Get());
+		if (InsertIndex == INDEX_NONE)
+		{
+			InsertIndex = 0;
+		}
+		else
+		{
+			InsertIndex++;
 		}
-	}
-}
 
-void FStreamingLevelModel::OnDrop(const TSharedPtr<FLevelDragDropOp>& Op)
-{
+		WorldStreamingLevels.Insert(DropStreamingLevels, InsertIndex);
+		CurrentWorld->MarkPackageDirty();
+			
+		// Force levels list refresh
+		LevelCollectionModel.PopulateLevelsList();
+	}
 }
 
 bool FStreamingLevelModel::IsGoodToDrop(const TSharedPtr<FLevelDragDropOp>& Op) const
 {
-	return false;
+	return true;
 }
 
 //bool FStreamingLevelModel::IsReadOnly() const
diff --git a/Engine/Source/Editor/WorldBrowser/Private/StreamingLevels/StreamingLevelModel.h b/Engine/Source/Editor/WorldBrowser/Private/StreamingLevels/StreamingLevelModel.h
--- a/Engine/Source/Editor/WorldBrowser/Private/StreamingLevels/StreamingLevelModel.h
+++ b/Engine/Source/Editor/WorldBrowser/Private/StreamingLevels/StreamingLevelModel.h
@@ -38,12 +38,6 @@ public:
 	virtual UClass* GetStreamingClass() const override;
 	// FLevelModel interface end
 		
-	/** @return the StreamingLevelIndex */
-	int GetStreamingLevelIndex() { return StreamingLevelIndex; }
-
-	/** Refreshes the StreamingLevelIndex by checking or its position in its OwningWorld's Level array */
-	void RefreshStreamingLevelIndex();
-
 	/** @return The ULevelStreaming this viewmodel contains*/
 	const TWeakObjectPtr< ULevelStreaming > GetLevelStreaming();
 
@@ -66,9 +60,6 @@ private:
 	/** The LevelStreaming this object represents */
 	TWeakObjectPtr< ULevelStreaming > LevelStreaming;
 
-	/** The index of this level within its parent UWorld's StreamingLevels Array;  Used for manual sorting */
-	int StreamingLevelIndex;
-
 	/** Whether underlying streaming level object has a valid package name */
 	bool bHasValidPackageName;
 };
\ No newline at end of file
