Adding Level Color back into the Levels Panel, fixed the color being displayed gamma corrected which didn't line up with the actual visualization.
Refactored some of the UI code for ULevelStreaming to use PostEditChangeProperty so it's more reuseable.

diff --git a/Engine/Source/Editor/DetailCustomizations/Private/MathStructCustomizations.cpp b/Engine/Source/Editor/DetailCustomizations/Private/MathStructCustomizations.cpp
--- a/Engine/Source/Editor/DetailCustomizations/Private/MathStructCustomizations.cpp
+++ b/Engine/Source/Editor/DetailCustomizations/Private/MathStructCustomizations.cpp
@@ -356,6 +356,7 @@ void FColorStructCustomization::MakeHeaderRow( TSharedRef<class IPropertyHandle>
 	// We'll set up reset to default ourselves
 	const bool bDisplayResetToDefault = false;
 	const FString DisplayNameOverride = TEXT("");
+	const bool bInterpretAsLinear = StructPropertyHandle->HasMetaData("InterpretAsLinear");
 
 	FSlateFontInfo NormalText = IDetailLayoutBuilder::GetDetailFont();
 
@@ -379,6 +380,7 @@ void FColorStructCustomization::MakeHeaderRow( TSharedRef<class IPropertyHandle>
 				.Color( this, &FColorStructCustomization::OnGetColorForColorBlock )
 				.ShowBackgroundForAlpha(true)
 				.IgnoreAlpha( bIgnoreAlpha )
+				.UseSRGB(!bInterpretAsLinear)
 				.OnMouseButtonDown( this, &FColorStructCustomization::OnMouseButtonDownColorBlock )
 				.Size( FVector2D( 35.0f, 12.0f ) )
 			]
@@ -401,6 +403,7 @@ void FColorStructCustomization::MakeHeaderRow( TSharedRef<class IPropertyHandle>
 			.Color( this, &FColorStructCustomization::OnGetColorForColorBlock )
 			.ShowBackgroundForAlpha(false)
 			.IgnoreAlpha(true)
+			.UseSRGB(!bInterpretAsLinear)
 			.OnMouseButtonDown( this, &FColorStructCustomization::OnMouseButtonDownColorBlock )
 			.Size( FVector2D( 35.0f, 12.0f ) )
 		]
diff --git a/Engine/Source/Editor/WorldBrowser/Private/LevelModel.cpp b/Engine/Source/Editor/WorldBrowser/Private/LevelModel.cpp
--- a/Engine/Source/Editor/WorldBrowser/Private/LevelModel.cpp
+++ b/Engine/Source/Editor/WorldBrowser/Private/LevelModel.cpp
@@ -333,6 +333,25 @@ void FLevelModel::SetLevelTranslationDelta(FVector2D InAbsoluteDelta)
 	}
 }
 
+bool FLevelModel::SupportsLevelColor() const
+{
+	// By default Levels don't support a level color
+	return false;
+}
+
+FColor FLevelModel::GetLevelColor() const
+{
+	// Returns Constant color, base classes will override this
+	// Currently not all base classes have the requisite support, so I've not made it pure virtual.
+	return FColor::White;
+}
+
+void FLevelModel::SetLevelColor(FColor InColor)
+{
+	// Does nothing, base classes will override this
+	// Currently not all base classes have the requisite support, so I've not made it pure virtual.
+}
+
 bool FLevelModel::HasKismet() const
 {
 	return (GetLevelObject() != NULL);
diff --git a/Engine/Source/Editor/WorldBrowser/Private/LevelModel.h b/Engine/Source/Editor/WorldBrowser/Private/LevelModel.h
--- a/Engine/Source/Editor/WorldBrowser/Private/LevelModel.h
+++ b/Engine/Source/Editor/WorldBrowser/Private/LevelModel.h
@@ -138,6 +138,15 @@ public:
 
 	/** Sets new translation delta to this model and all descendants*/
 	void SetLevelTranslationDelta(FVector2D InAbsoluteDelta);
+
+	/** @return Whether level supports a Level Color for visualization. (Show -> Advanced -> Level Coloration) */
+	virtual bool SupportsLevelColor() const;
+
+	/** @return level color, used for visualization. (Show -> Advanced -> Level Coloration) */
+	virtual FColor GetLevelColor() const;
+
+	/** Sets level color, used for visualization. (Show -> Advanced -> Level Coloration) */
+	virtual void SetLevelColor(FColor InColor);
 	
 	/**	@return Whether level has associated blueprint script */
 	bool HasKismet() const;
diff --git a/Engine/Source/Editor/WorldBrowser/Private/SWorldHierarchy.cpp b/Engine/Source/Editor/WorldBrowser/Private/SWorldHierarchy.cpp
--- a/Engine/Source/Editor/WorldBrowser/Private/SWorldHierarchy.cpp
+++ b/Engine/Source/Editor/WorldBrowser/Private/SWorldHierarchy.cpp
@@ -104,6 +104,15 @@ public:
 				[
 					SNew(STextBlock)
 						.ToolTipText(NSLOCTEXT("WorldBrowser", "Save", "Save this Level"))
+				]
+
+			/** Level color column */
+			+ SHeaderRow::Column(HierarchyColumns::ColumnID_Color)
+				.FixedWidth(24.0f)
+				.HeaderContent()
+				[
+					SNew(STextBlock)
+					.ToolTipText(NSLOCTEXT("WorldBrowser", "Color", "Color used for visualization of Level"))
 				];
 
 
diff --git a/Engine/Source/Editor/WorldBrowser/Private/SWorldHierarchyItem.cpp b/Engine/Source/Editor/WorldBrowser/Private/SWorldHierarchyItem.cpp
--- a/Engine/Source/Editor/WorldBrowser/Private/SWorldHierarchyItem.cpp
+++ b/Engine/Source/Editor/WorldBrowser/Private/SWorldHierarchyItem.cpp
@@ -104,17 +104,44 @@ TSharedRef< SWidget > SWorldHierarchyItem::GenerateWidgetForColumn( const FName
 			]
 		;
 	}
+	else if (ColumnID == HierarchyColumns::ColumnID_Color)
+	{
+		if (LevelModel->SupportsLevelColor())
+		{
+			TableRowContent =
+				SAssignNew(ColorButton, SButton)
+				.ContentPadding(0)
+				.ButtonStyle(FEditorStyle::Get(), "ToggleButton")
+				.IsEnabled(true)
+				.OnClicked(this, &SWorldHierarchyItem::OnChangeColor)
+				.ToolTipText(LOCTEXT("LevelColorButtonToolTip", "Change Level Color"))
+				.HAlign(HAlign_Center)
+				.VAlign(VAlign_Center)
+				.Content()
+				[
+					SNew(SImage)
+					.ColorAndOpacity(this, &SWorldHierarchyItem::GetDrawColor)
+					.Image(this, &SWorldHierarchyItem::GetLevelColorBrush)
+				]
+			;
+		}
+		else
+		{
+			TableRowContent =
+				SNew(SHorizontalBox);
+		}
+	}
 	else if (ColumnID == HierarchyColumns::ColumnID_Kismet)
 	{
-		TableRowContent = 
+		TableRowContent =
 			SAssignNew(KismetButton, SButton)
 			.ContentPadding(0)
 			.ButtonStyle(FEditorStyle::Get(), "ToggleButton")
 			.IsEnabled(this, &SWorldHierarchyItem::IsKismetEnabled)
 			.OnClicked(this, &SWorldHierarchyItem::OnOpenKismet)
 			.ToolTipText(LOCTEXT("KismetButtonToolTip", "Open Level Blueprint"))
-			.HAlign( HAlign_Center )
-			.VAlign( VAlign_Center )
+			.HAlign(HAlign_Center)
+			.VAlign(VAlign_Center)
 			.Content()
 			[
 				SNew(SImage)
@@ -191,6 +218,11 @@ bool SWorldHierarchyItem::IsKismetEnabled() const
 	return LevelModel->HasKismet();
 }
 
+FSlateColor SWorldHierarchyItem::GetDrawColor() const
+{
+	return FSlateColor(LevelModel->GetLevelColor());
+}
+
 FReply SWorldHierarchyItem::OnToggleVisibility()
 {
 	FLevelModelList LevelList; 
@@ -241,6 +273,52 @@ FReply SWorldHierarchyItem::OnOpenKismet()
 	return FReply::Handled();
 }
 
+void SWorldHierarchyItem::OnSetColorFromColorPicker(FLinearColor NewColor)
+{
+	LevelModel->SetLevelColor(NewColor.ToFColor(true));
+}
+
+void SWorldHierarchyItem::OnColorPickerCancelled(FLinearColor OriginalColor)
+{
+	LevelModel->SetLevelColor(OriginalColor.ToFColor(true));
+}
+
+void SWorldHierarchyItem::OnColorPickerInteractiveBegin()
+{
+	GEditor->BeginTransaction(NSLOCTEXT("FColorStructCustomization", "SetColorProperty", "Edit Level Draw Color"));
+}
+
+void SWorldHierarchyItem::OnColorPickerInteractiveEnd()
+{
+	GEditor->EndTransaction();
+}
+
+FReply SWorldHierarchyItem::OnChangeColor()
+{
+	if (LevelModel.IsValid())
+	{
+		// Initialize the color data for the picker window.
+		FColor NewColor = LevelModel->GetLevelColor();
+		TArray<FColor*> ColorArray;
+		ColorArray.Add(&NewColor);
+
+		FColorPickerArgs PickerArgs;
+		PickerArgs.DisplayGamma = TAttribute<float>::Create(TAttribute<float>::FGetter::CreateUObject(GEngine, &UEngine::GetDisplayGamma));
+		PickerArgs.ColorArray = &ColorArray;
+		PickerArgs.bOnlyRefreshOnMouseUp = false;
+		PickerArgs.bOnlyRefreshOnOk = false;
+		PickerArgs.OnColorCommitted = FOnLinearColorValueChanged::CreateSP(this, &SWorldHierarchyItem::OnSetColorFromColorPicker);
+		PickerArgs.OnColorPickerCancelled = FOnColorPickerCancelled::CreateSP(this, &SWorldHierarchyItem::OnColorPickerCancelled);
+		PickerArgs.OnInteractivePickBegin = FSimpleDelegate::CreateSP(this, &SWorldHierarchyItem::OnColorPickerInteractiveBegin);
+		PickerArgs.OnInteractivePickEnd = FSimpleDelegate::CreateSP(this, &SWorldHierarchyItem::OnColorPickerInteractiveEnd);
+		PickerArgs.ParentWidget = AsShared();
+
+		OpenColorPicker(PickerArgs);
+	}
+
+	return FReply::Handled();
+}
+
 FReply SWorldHierarchyItem::OnItemDragDetected(const FGeometry& MyGeometry, const FPointerEvent& MouseEvent)
 {
 	if (MouseEvent.IsMouseButtonDown(EKeys::LeftMouseButton))
@@ -489,4 +567,9 @@ const FSlateBrush* SWorldHierarchyItem::GetLevelKismetBrush() const
 	}	
 }
 
+const FSlateBrush* SWorldHierarchyItem::GetLevelColorBrush() const
+{
+	return FEditorStyle::GetBrush("Level.ColorIcon40x");
+}
+
 #undef LOCTEXT_NAMESPACE
diff --git a/Engine/Source/Editor/WorldBrowser/Private/SWorldHierarchyItem.h b/Engine/Source/Editor/WorldBrowser/Private/SWorldHierarchyItem.h
--- a/Engine/Source/Editor/WorldBrowser/Private/SWorldHierarchyItem.h
+++ b/Engine/Source/Editor/WorldBrowser/Private/SWorldHierarchyItem.h
@@ -9,6 +9,7 @@ namespace HierarchyColumns
 	static const FName ColumnID_Lock( "Lock" );
 	static const FName ColumnID_SCCStatus( "SCC_Status" );
 	static const FName ColumnID_Save( "Save" );
+	static const FName ColumnID_Color("Color");
 	static const FName ColumnID_Kismet( "Blueprint" );
 	static const FName ColumnID_ActorCount( "ActorCount" );
 	static const FName ColumnID_LightmassSize( "LightmassSize" );
@@ -55,6 +56,9 @@ private:
 	bool IsLockEnabled() const;
 	bool IsVisibilityEnabled() const;
 	bool IsKismetEnabled() const;
+
+	/** Get DrawColor for the level */
+	FSlateColor GetDrawColor() const;
 		
 	/**
 	 *	Called when the user clicks on the visibility icon for a Level's item widget
@@ -83,6 +87,19 @@ private:
 	 *	@return	A reply that indicated whether this event was handled.
 	 */
 	FReply OnOpenKismet();
+
+	/**
+	*	Called when the user clicks on the color icon for a Level's item widget
+	*
+	*	@return	A reply that indicated whether this event was handled.
+	*/
+	FReply OnChangeColor();
+
+	/** Callback invoked from the color picker */
+	void OnSetColorFromColorPicker(FLinearColor NewColor);
+	void OnColorPickerCancelled(FLinearColor OriginalColor);
+	void OnColorPickerInteractiveBegin();
+	void OnColorPickerInteractiveEnd();
 	
 	/**
 	 *  @return The text of level name while it is not being edited
@@ -131,6 +148,15 @@ private:
 	 *	@return	The SlateBrush representing the Level's kismet state
 	 */
 	const FSlateBrush* GetLevelKismetBrush() const;
+
+	/**
+	*	Called to get the Slate Image Brush representing the color of
+	*	the Level this row widget represents
+	*
+	*	@return	The SlateBrush representing the Level's color
+	*/
+	const FSlateBrush* GetLevelColorBrush() const;
+
 	
 	/** */
 	FString GetLevelLockToolTip() const;
diff --git a/Engine/Source/Editor/WorldBrowser/Private/StreamingLevels/StreamingLevelModel.cpp b/Engine/Source/Editor/WorldBrowser/Private/StreamingLevels/StreamingLevelModel.cpp
--- a/Engine/Source/Editor/WorldBrowser/Private/StreamingLevels/StreamingLevelModel.cpp
+++ b/Engine/Source/Editor/WorldBrowser/Private/StreamingLevels/StreamingLevelModel.cpp
@@ -64,6 +64,34 @@ FName FStreamingLevelModel::GetLongPackageName() const
 	}
 }
 
+bool FStreamingLevelModel::SupportsLevelColor() const
+{
+	return LevelStreaming.IsValid();
+}
+
+FColor FStreamingLevelModel::GetLevelColor() const
+{
+	if (LevelStreaming.IsValid())
+	{
+		return LevelStreaming.Get()->DrawColor;
+	}
+	return FLevelModel::GetLevelColor();
+}
+
+void FStreamingLevelModel::SetLevelColor(FColor InColor)
+{
+	if (LevelStreaming.IsValid())
+	{
+		UProperty* DrawColorProperty = FindField<UProperty>(LevelStreaming->GetClass(), "DrawColor");
+		LevelStreaming->PreEditChange(DrawColorProperty);
+
+		LevelStreaming.Get()->DrawColor = InColor;
+
+		FPropertyChangedEvent PropertyChangedEvent(DrawColorProperty, false, EPropertyChangeType::ValueSet);
+		LevelStreaming->PostEditChangeProperty(PropertyChangedEvent);
+	}
+}
+
 void FStreamingLevelModel::Update()
 {
 	UpdatePackageFileAvailability();
diff --git a/Engine/Source/Editor/WorldBrowser/Private/StreamingLevels/StreamingLevelModel.h b/Engine/Source/Editor/WorldBrowser/Private/StreamingLevels/StreamingLevelModel.h
--- a/Engine/Source/Editor/WorldBrowser/Private/StreamingLevels/StreamingLevelModel.h
+++ b/Engine/Source/Editor/WorldBrowser/Private/StreamingLevels/StreamingLevelModel.h
@@ -32,6 +32,9 @@ public:
 	virtual ULevel* GetLevelObject() const override;
 	virtual FName GetAssetName() const override;
 	virtual FName GetLongPackageName() const override;
+	virtual bool SupportsLevelColor() const override;
+	virtual FColor GetLevelColor() const override;
+	virtual void SetLevelColor(FColor InColor) override;
 	virtual void Update() override;
 	virtual void OnDrop(const TSharedPtr<FLevelDragDropOp>& Op) override;
 	virtual bool IsGoodToDrop(const TSharedPtr<FLevelDragDropOp>& Op) const override;
diff --git a/Engine/Source/Runtime/Engine/Classes/Engine/LevelStreaming.h b/Engine/Source/Runtime/Engine/Classes/Engine/LevelStreaming.h
--- a/Engine/Source/Runtime/Engine/Classes/Engine/LevelStreaming.h
+++ b/Engine/Source/Runtime/Engine/Classes/Engine/LevelStreaming.h
@@ -126,7 +126,7 @@ class ULevelStreaming : public UObject
 	uint32 bDrawOnLevelStatusMap:1;
 
 	/** The level color used for visualization. (Show -> Advanced -> Level Coloration) */
-	UPROPERTY(EditAnywhere, Category=LevelStreaming)
+	UPROPERTY(EditAnywhere, Category = LevelStreaming, meta=(InterpretAsLinear))
 	FColor DrawColor;
 
 	/** The level streaming volumes bound to this level.																		*/
@@ -151,7 +151,7 @@ class ULevelStreaming : public UObject
 	virtual void PostLoad() override;
 	virtual void Serialize( FArchive& Ar ) override;
 #if WITH_EDITOR
-	virtual void PostEditChangeChainProperty(FPropertyChangedChainEvent& PropertyChangedEvent) override;
+	virtual void PostEditChangeProperty(FPropertyChangedEvent& PropertyChangedEvent) override;
 	
 	/** Remove duplicates in EditorStreamingVolumes list*/
 	void RemoveStreamingVolumeDuplicates();
diff --git a/Engine/Source/Runtime/Engine/Private/LevelStreaming.cpp b/Engine/Source/Runtime/Engine/Private/LevelStreaming.cpp
--- a/Engine/Source/Runtime/Engine/Private/LevelStreaming.cpp
+++ b/Engine/Source/Runtime/Engine/Private/LevelStreaming.cpp
@@ -754,46 +754,43 @@ FBox ULevelStreaming::GetStreamingVolumeBounds()
 }
 
 #if WITH_EDITOR
-void ULevelStreaming::PostEditChangeChainProperty(FPropertyChangedChainEvent& PropertyChangedEvent)
+void ULevelStreaming::PostEditChangeProperty(FPropertyChangedEvent& PropertyChangedEvent)
 {
-	if ( PropertyChangedEvent.PropertyChain.Num() > 0 )
+	UProperty* OutermostProperty = PropertyChangedEvent.Property;
+	if ( OutermostProperty != NULL )
 	{
-		UProperty* OutermostProperty = PropertyChangedEvent.PropertyChain.GetHead()->GetValue();
-		if ( OutermostProperty != NULL )
+		const FName PropertyName = OutermostProperty->GetFName();
+		if (PropertyName == GET_MEMBER_NAME_CHECKED(ULevelStreaming, LevelTransform))
 		{
-			const FName PropertyName = OutermostProperty->GetFName();
-			if ( PropertyName == TEXT("LevelTransform") )
-			{
-				GetWorld()->UpdateLevelStreaming();
-			}
+			GetWorld()->UpdateLevelStreaming();
+		}
 
-			if (PropertyName == GET_MEMBER_NAME_CHECKED(ULevelStreaming, EditorStreamingVolumes))
-			{
-				RemoveStreamingVolumeDuplicates();
+		if (PropertyName == GET_MEMBER_NAME_CHECKED(ULevelStreaming, EditorStreamingVolumes))
+		{
+			RemoveStreamingVolumeDuplicates();
 
-				// Update levels references in each streaming volume 
-				for (TActorIterator<ALevelStreamingVolume> It(GetWorld()); It; ++It)
-				{
-					(*It)->UpdateStreamingLevelsRefs();
-				}
+			// Update levels references in each streaming volume 
+			for (TActorIterator<ALevelStreamingVolume> It(GetWorld()); It; ++It)
+			{
+				(*It)->UpdateStreamingLevelsRefs();
 			}
+		}
 
-			else if( PropertyName == TEXT( "DrawColor" ) )
+		else if (PropertyName == GET_MEMBER_NAME_CHECKED(ULevelStreaming, DrawColor))
+		{
+			// Make sure the level's DrawColor change is applied immediately by reregistering the
+			// components of the actor's in the level
+			if( LoadedLevel != NULL )
 			{
-				// Make sure the level's DrawColor change is applied immediately by reregistering the
-				// components of the actor's in the level
-				if( LoadedLevel != NULL )
+				UPackage* Package = LoadedLevel->GetOutermost();
+				for( TObjectIterator<UActorComponent> It; It; ++It )
 				{
-					UPackage* Package = LoadedLevel->GetOutermost();
-					for( TObjectIterator<UActorComponent> It; It; ++It )
+					if( It->IsIn( Package ) )
 					{
-						if( It->IsIn( Package ) )
+						UActorComponent* ActorComponent = Cast<UActorComponent>( *It );
+						if( ActorComponent )
 						{
-							UActorComponent* ActorComponent = Cast<UActorComponent>( *It );
-							if( ActorComponent )
-							{
-								ActorComponent->RecreateRenderState_Concurrent();
-							}
+							ActorComponent->RecreateRenderState_Concurrent();
 						}
 					}
 				}
@@ -801,7 +798,7 @@ void ULevelStreaming::PostEditChangeChainProperty(FPropertyChangedChainEvent& Pr
 		}
 	}
 
-	Super::PostEditChangeChainProperty(PropertyChangedEvent);
+	Super::PostEditChangeProperty(PropertyChangedEvent);
 }
 
 void ULevelStreaming::RemoveStreamingVolumeDuplicates()
