Viewport actor menu command "Add/Remove level selection" works correcly with both old levels browser and a new one.

diff --git a/Engine/Source/Editor/UnrealEd/Private/Editor.cpp b/Engine/Source/Editor/UnrealEd/Private/Editor.cpp
--- a/Engine/Source/Editor/UnrealEd/Private/Editor.cpp
+++ b/Engine/Source/Editor/UnrealEd/Private/Editor.cpp
@@ -3174,7 +3174,14 @@ void UEditorEngine::SelectLevelInLevelBrowser( bool bDeselectOthers )
 	}
 
 	FLevelEditorModule& LevelEditorModule = FModuleManager::LoadModuleChecked<FLevelEditorModule>("LevelEditor");
-	LevelEditorModule.SummonLevelBrowser();
+	if (FParse::Param(FCommandLine::Get(), TEXT("oldlevels")))
+	{
+		LevelEditorModule.SummonLevelBrowser();
+	}
+	else
+	{
+		LevelEditorModule.SummonWorldBrowserHierarchy();
+	}
 }
 
 void UEditorEngine::DeselectLevelInLevelBrowser()
@@ -3187,9 +3194,16 @@ void UEditorEngine::DeselectLevelInLevelBrowser()
 			Actor->GetWorld()->DeSelectLevel( Actor->GetLevel() );
 		}
 	}
-
+	
 	FLevelEditorModule& LevelEditorModule = FModuleManager::LoadModuleChecked<FLevelEditorModule>("LevelEditor");
-	LevelEditorModule.SummonLevelBrowser();
+	if (FParse::Param(FCommandLine::Get(), TEXT("oldlevels")))
+	{
+		LevelEditorModule.SummonLevelBrowser();
+	}
+	else
+	{
+		LevelEditorModule.SummonWorldBrowserHierarchy();
+	}
 }
 
 void UEditorEngine::SelectAllActorsControlledByMatinee()
diff --git a/Engine/Source/Editor/WorldBrowser/Private/LevelCollectionModel.cpp b/Engine/Source/Editor/WorldBrowser/Private/LevelCollectionModel.cpp
--- a/Engine/Source/Editor/WorldBrowser/Private/LevelCollectionModel.cpp
+++ b/Engine/Source/Editor/WorldBrowser/Private/LevelCollectionModel.cpp
@@ -24,6 +24,7 @@ FLevelCollectionModel::FLevelCollectionModel(UEditorEngine* InEditor)
 	, bCanExecuteSCCCheckIn(false)
 	, bCanExecuteSCC(false)
 	, bSelectionHasChanged(true)
+	, bUpdatingLevelsSelection(false)
 {
 }
 
@@ -37,6 +38,10 @@ FLevelCollectionModel::~FLevelCollectionModel()
 	FEditorSupportDelegates::RedrawAllViewports.RemoveAll(this);
 	Editor->OnLevelActorAdded().RemoveAll(this);
 	Editor->OnLevelActorDeleted().RemoveAll(this);
+	if (CurrentWorld.IsValid())
+	{
+		CurrentWorld->OnSelectedLevelsChanged().RemoveAll(this);
+	}
 }
 
 void FLevelCollectionModel::Initialize(UWorld* InWorld)
@@ -53,6 +58,7 @@ void FLevelCollectionModel::Initialize(UWorld* InWorld)
 	Editor->OnLevelActorDeleted().AddSP( this, &FLevelCollectionModel::OnLevelActorDeleted);
 	USelection::SelectionChangedEvent.AddSP(this, &FLevelCollectionModel::OnActorSelectionChanged);
 	SelectionChanged.AddSP(this, &FLevelCollectionModel::OnActorOrLevelSelectionChanged);
+	CurrentWorld->OnSelectedLevelsChanged().AddSP(this, &FLevelCollectionModel::OnLevelsSelectionChangedOutside);
 
 	PopulateLevelsList();
 }
@@ -334,6 +340,22 @@ void FLevelCollectionModel::SetSelectedLevels(const FLevelModelList& InList)
 	OnLevelsSelectionChanged();
 }
 
+void FLevelCollectionModel::SetSelectedLevelsFromWorld()
+{
+	TArray<ULevel*>& SelectedLevelObjects = CurrentWorld->GetSelectedLevels();
+	FLevelModelList LevelsToSelect;
+	for (ULevel* LevelObject : SelectedLevelObjects)
+	{
+		TSharedPtr<FLevelModel> LevelModel = FindLevelModel(LevelObject);
+		if (LevelModel.IsValid())
+		{
+			LevelsToSelect.Add(LevelModel);
+		}
+	}
+
+	SetSelectedLevels(LevelsToSelect);
+}
+
 TSharedPtr<FLevelModel> FLevelCollectionModel::FindLevelModel(ULevel* InLevel) const
 {
 	if (InLevel)
@@ -1584,6 +1606,13 @@ void FLevelCollectionModel::OnLevelsCollectionChanged()
 
 void FLevelCollectionModel::OnLevelsSelectionChanged()
 {
+	if (bUpdatingLevelsSelection)
+	{
+		return;
+	}
+	
+	TGuardValue<bool> UpdateGuard(bUpdatingLevelsSelection, true);
+	
 	// Pass the list we just created to the world to set the selection
 	CurrentWorld->SetSelectedLevels(
 		GetLevelObjectList(SelectedLevelsList)
@@ -1608,6 +1637,14 @@ void FLevelCollectionModel::OnLevelsSelectionChanged()
 	BroadcastSelectionChanged();
 }
 
+void FLevelCollectionModel::OnLevelsSelectionChangedOutside()
+{
+	if (!bUpdatingLevelsSelection)
+	{
+		SetSelectedLevelsFromWorld();
+	}
+}
+
 void FLevelCollectionModel::OnLevelsHierarchyChanged()
 {
 	BroadcastHierarchyChanged();
diff --git a/Engine/Source/Editor/WorldBrowser/Private/LevelCollectionModel.h b/Engine/Source/Editor/WorldBrowser/Private/LevelCollectionModel.h
--- a/Engine/Source/Editor/WorldBrowser/Private/LevelCollectionModel.h
+++ b/Engine/Source/Editor/WorldBrowser/Private/LevelCollectionModel.h
@@ -72,6 +72,9 @@ public:
 
 	/**	Sets selected level list */
 	void SetSelectedLevels(const FLevelModelList& InList);
+	
+	/**	Sets selection to a levels that is currently marked as selected in UWorld */
+	void SetSelectedLevelsFromWorld();
 
 	/**	@return	Found level model which represents specified level object */
 	TSharedPtr<FLevelModel> FindLevelModel(ULevel* InLevel) const;
@@ -424,6 +427,9 @@ protected:
 	
 	/** Called whenever level selection has been changed */
 	virtual void OnLevelsSelectionChanged();
+
+	/** Called whenever level selection has been changed outside of this module, usually via World->SetSelectedLevels */
+	void OnLevelsSelectionChangedOutside();
 	
 	/** Called whenever level collection hierarchy has been changed */
 	virtual void OnLevelsHierarchyChanged();
@@ -515,6 +521,9 @@ protected:
 
 	/** Flag for whether the selection of levels or actors has changed */
 	bool								bSelectionHasChanged;
+
+	/** Guard to avoid recursive level selection updates */
+	bool								bUpdatingLevelsSelection;
 };
 
 //
diff --git a/Engine/Source/Editor/WorldBrowser/Private/StreamingLevels/StreamingLevelCollectionModel.cpp b/Engine/Source/Editor/WorldBrowser/Private/StreamingLevels/StreamingLevelCollectionModel.cpp
--- a/Engine/Source/Editor/WorldBrowser/Private/StreamingLevels/StreamingLevelCollectionModel.cpp
+++ b/Engine/Source/Editor/WorldBrowser/Private/StreamingLevels/StreamingLevelCollectionModel.cpp
@@ -68,8 +68,11 @@ void FStreamingLevelCollectionModel::OnLevelsCollectionChanged()
 			LevelModel->SetParent(PersistentLevelModel);
 		}
 	}
-
+	
 	FLevelCollectionModel::OnLevelsCollectionChanged();
+
+	// Sync levels selection to world
+	SetSelectedLevelsFromWorld();
 }
 
 void FStreamingLevelCollectionModel::OnLevelsSelectionChanged()
diff --git a/Engine/Source/Editor/WorldBrowser/Private/Tiles/WorldTileCollectionModel.cpp b/Engine/Source/Editor/WorldBrowser/Private/Tiles/WorldTileCollectionModel.cpp
--- a/Engine/Source/Editor/WorldBrowser/Private/Tiles/WorldTileCollectionModel.cpp
+++ b/Engine/Source/Editor/WorldBrowser/Private/Tiles/WorldTileCollectionModel.cpp
@@ -737,6 +737,9 @@ void FWorldTileCollectionModel::OnLevelsCollectionChanged()
 	}
 
 	FLevelCollectionModel::OnLevelsCollectionChanged();
+
+	// Sync levels selection to world
+	SetSelectedLevelsFromWorld();
 }
 
 TSharedPtr<FWorldTileModel> FWorldTileCollectionModel::AddLevelFromTile(int32 TileIdx)
