Removed unnecessary SWidget::OnHitTest override from world composition tiles

diff --git a/Engine/Source/Editor/WorldBrowser/Private/Tiles/SWorldTileItem.cpp b/Engine/Source/Editor/WorldBrowser/Private/Tiles/SWorldTileItem.cpp
--- a/Engine/Source/Editor/WorldBrowser/Private/Tiles/SWorldTileItem.cpp
+++ b/Engine/Source/Editor/WorldBrowser/Private/Tiles/SWorldTileItem.cpp
@@ -439,21 +439,8 @@ int32 SWorldTileItem::OnPaint(const FPaintArgs& Args, const FGeometry& AllottedG
 
 FReply SWorldTileItem::OnMouseButtonDoubleClick(const FGeometry& InMyGeometry, const FPointerEvent& InMouseEvent)
 {
-	if (OnHitTest(InMyGeometry, InMouseEvent.GetScreenSpacePosition()))
-	{
-		TileModel->MakeLevelCurrent();
-		return FReply::Handled();
-	}
-	
-	return FReply::Unhandled();
-}
-
-bool SWorldTileItem::OnHitTest(const FGeometry& MyGeometry, FVector2D InAbsoluteCursorPosition)
-{
-	FVector2D CursorOffset = (InAbsoluteCursorPosition - MyGeometry.AbsolutePosition)/MyGeometry.Scale;
-	FVector2D LevelPos = TileModel->GetLevelPosition2D();
-	FVector2D CursorWorldPos = LevelPos + CursorOffset;
-	return TileModel->HitTest2D(CursorWorldPos);
+	TileModel->MakeLevelCurrent();
+	return FReply::Handled();
 }
 
 FString SWorldTileItem::GetLevelNameText() const
diff --git a/Engine/Source/Editor/WorldBrowser/Private/Tiles/SWorldTileItem.h b/Engine/Source/Editor/WorldBrowser/Private/Tiles/SWorldTileItem.h
--- a/Engine/Source/Editor/WorldBrowser/Private/Tiles/SWorldTileItem.h
+++ b/Engine/Source/Editor/WorldBrowser/Private/Tiles/SWorldTileItem.h
@@ -88,7 +88,6 @@ private:
 	// SWidget interface start
 	virtual TSharedPtr<IToolTip> GetToolTip() override;
 	virtual FVector2D ComputeDesiredSize() const override;
-	virtual bool OnHitTest(const FGeometry& MyGeometry, FVector2D InAbsoluteCursorPosition) override;
 	virtual int32 OnPaint(const FPaintArgs& Args, const FGeometry& AllottedGeometry, const FSlateRect& MyClippingRect, FSlateWindowElementList& OutDrawElements, int32 LayerId, const FWidgetStyle& InWidgetStyle, bool bParentEnabled) const override;
 	virtual FReply OnMouseButtonDoubleClick(const FGeometry& InMyGeometry, const FPointerEvent& InMouseEvent) override;
 	// SWidget interface end
diff --git a/Engine/Source/Editor/WorldBrowser/Private/Tiles/WorldTileModel.cpp b/Engine/Source/Editor/WorldBrowser/Private/Tiles/WorldTileModel.cpp
--- a/Engine/Source/Editor/WorldBrowser/Private/Tiles/WorldTileModel.cpp
+++ b/Engine/Source/Editor/WorldBrowser/Private/Tiles/WorldTileModel.cpp
@@ -120,32 +120,6 @@ FName FWorldTileModel::GetLongPackageName() const
 	return TileDetails->PackageName;
 }
 
-bool FWorldTileModel::HitTest2D(const FVector2D& Point) const
-{
-	if (0 && IsLandscapeBased())
-	{
-		FVector2D LandscapePos = GetLevelCurrentPosition() - FVector2D(TileDetails->Bounds.GetExtent());
-
-		for (int32 CompIndex = 0; CompIndex < LandscapeComponentsXY.Num(); ++CompIndex)
-		{
-			const FIntPoint& CompCoords = LandscapeComponentsXY[CompIndex];
-			FVector2D Min = LandscapePos + LandscapeComponentSize*FVector2D(CompCoords);
-			FVector2D Max = Min + LandscapeComponentSize;
-
-			if ((Point.X > Min.X) && (Point.X < Max.X) && (Point.Y > Min.Y) && (Point.Y < Max.Y))
-			{
-				return true;
-			}
-		}
-
-		return false;
-	}
-	else
-	{
-		return true;
-	}
-}
-
 FVector2D FWorldTileModel::GetLevelPosition2D() const
 {
 	if (TileDetails->Bounds.IsValid)
diff --git a/Engine/Source/Editor/WorldBrowser/Private/Tiles/WorldTileModel.h b/Engine/Source/Editor/WorldBrowser/Private/Tiles/WorldTileModel.h
--- a/Engine/Source/Editor/WorldBrowser/Private/Tiles/WorldTileModel.h
+++ b/Engine/Source/Editor/WorldBrowser/Private/Tiles/WorldTileModel.h
@@ -75,7 +75,6 @@ public:
 	virtual void Update() override;
 	virtual void LoadLevel() override;
 	virtual void SetVisible(bool bVisible) override;
-	virtual bool HitTest2D(const FVector2D& Point) const override;
 	virtual FVector2D GetLevelPosition2D() const override;
 	virtual FVector2D GetLevelSize2D() const override;
 	virtual void OnDrop(const TSharedPtr<FLevelDragDropOp>& Op) override;
