Fixed world composition view regression

diff --git a/Engine/Source/Editor/WorldBrowser/Private/LevelCollectionModel.h b/Engine/Source/Editor/WorldBrowser/Private/LevelCollectionModel.h
--- a/Engine/Source/Editor/WorldBrowser/Private/LevelCollectionModel.h
+++ b/Engine/Source/Editor/WorldBrowser/Private/LevelCollectionModel.h
@@ -46,9 +46,6 @@ public:
 	/** Current world size  */
 	FIntPoint GetWorldSize() const { return WorldSize; }
 	
-	/** Current world origin location on XY plane  */
-	FIntPoint GetWorldOriginLocationXY() const { return FIntPoint(GetWorld()->OriginLocation.X, GetWorld()->OriginLocation.Y); }
-	
 	/**	@return	Root list of levels in hierarchy */
 	FLevelModelList& GetRootLevelList();
 
@@ -529,6 +526,12 @@ protected:
 	bool								bUpdatingLevelsSelection;
 };
 
+/** Current world origin location on XY plane  */
+FORCEINLINE FIntPoint GetWorldOriginLocationXY(UWorld* InWorld)
+{
+	return FIntPoint(InWorld->OriginLocation.X, InWorld->OriginLocation.Y);
+}
+
 //
 // Helper struct to temporally make specified UObject immune to dirtying
 //
diff --git a/Engine/Source/Editor/WorldBrowser/Private/Tiles/SWorldComposition.cpp b/Engine/Source/Editor/WorldBrowser/Private/Tiles/SWorldComposition.cpp
--- a/Engine/Source/Editor/WorldBrowser/Private/Tiles/SWorldComposition.cpp
+++ b/Engine/Source/Editor/WorldBrowser/Private/Tiles/SWorldComposition.cpp
@@ -772,13 +772,11 @@ protected:
 	/** Delegate callback: world origin is going to be moved. */
 	void PreWorldOriginOffset(UWorld* InWorld, FIntVector InSrcOrigin, FIntVector InDstOrigin)
 	{
-		if (WorldModel->GetWorld() != InWorld)
+		if (InWorld && (WorldModel->GetWorld() == InWorld || WorldModel->GetSimulationWorld() == InWorld))		
 		{
-			return;
+			FIntVector Offset = InDstOrigin - InSrcOrigin;
+			RequestScrollBy(-FVector2D(Offset.X, Offset.Y));
 		}
-		
-		FIntVector Offset = InDstOrigin - InSrcOrigin;
-		RequestScrollBy(-FVector2D(Offset.X, Offset.Y));
 	}
 	
 	/** Handles new item added to data source */
diff --git a/Engine/Source/Editor/WorldBrowser/Private/Tiles/WorldTileCollectionModel.cpp b/Engine/Source/Editor/WorldBrowser/Private/Tiles/WorldTileCollectionModel.cpp
--- a/Engine/Source/Editor/WorldBrowser/Private/Tiles/WorldTileCollectionModel.cpp
+++ b/Engine/Source/Editor/WorldBrowser/Private/Tiles/WorldTileCollectionModel.cpp
@@ -1828,7 +1828,7 @@ bool FWorldTileCollectionModel::GenerateLODLevels(FLevelModelList InLevelList, i
 		LODPackage->FileName = FName(*LODLevelFileName);
 				
 		// This is current actors offset from their original position
-		FVector ActorsOffset = FVector(TileModel->GetAbsoluteLevelPosition() - GetWorldOriginLocationXY());
+		FVector ActorsOffset = FVector(TileModel->GetAbsoluteLevelPosition() - GetWorldOriginLocationXY(GetWorld()));
 		if (GetWorld()->WorldComposition->bTemporallyDisableOriginTracking)
 		{
 			ActorsOffset = FVector::ZeroVector;
diff --git a/Engine/Source/Editor/WorldBrowser/Private/Tiles/WorldTileModel.cpp b/Engine/Source/Editor/WorldBrowser/Private/Tiles/WorldTileModel.cpp
--- a/Engine/Source/Editor/WorldBrowser/Private/Tiles/WorldTileModel.cpp
+++ b/Engine/Source/Editor/WorldBrowser/Private/Tiles/WorldTileModel.cpp
@@ -389,7 +389,7 @@ FVector2D FWorldTileModel::GetLevelCurrentPosition() const
 		FVector2D LevelLocalPosition(TileDetails->Bounds.GetCenter());
 		FIntPoint LevelOffset = GetAbsoluteLevelPosition();
 			
-		return LevelLocalPosition + FVector2D(LevelOffset - LevelCollectionModel.GetWorldOriginLocationXY()); 
+		return LevelLocalPosition + FVector2D(LevelOffset - GetWorldOriginLocationXY(CurrentWorld)); 
 	}
 
 	return FVector2D(0, 0);
@@ -481,7 +481,7 @@ void FWorldTileModel::Update()
 			{
 				FBox LevelWorldBounds = TileDetails->Bounds;
 				FIntPoint LevelAbolutePosition = GetAbsoluteLevelPosition();
-				FIntPoint LevelOffset = LevelAbolutePosition - LevelCollectionModel.GetWorldOriginLocationXY();
+				FIntPoint LevelOffset = LevelAbolutePosition - GetWorldOriginLocationXY(LevelCollectionModel.GetWorld());
 
 				TileDetails->Bounds = LevelWorldBounds.ShiftBy(-FVector(LevelOffset));
 			}
@@ -762,7 +762,7 @@ bool FWorldTileModel::CreateAdjacentLandscapeProxy(ALandscapeProxy* SourceLandsc
 		}
 
 		// Add source level position
-		FIntPoint IntOffset = FIntPoint(ProxyOffset.X, ProxyOffset.Y) + LevelCollectionModel.GetWorldOriginLocationXY();
+		FIntPoint IntOffset = FIntPoint(ProxyOffset.X, ProxyOffset.Y) + GetWorldOriginLocationXY(LevelCollectionModel.GetWorld());
 
 		// Move level with landscape proxy to desired position
 		SetLevelPosition(IntOffset);
