World composition: some move improvements against loops in hierarchies

diff --git a/Engine/Source/Editor/WorldBrowser/Private/Tiles/WorldTileModel.cpp b/Engine/Source/Editor/WorldBrowser/Private/Tiles/WorldTileModel.cpp
--- a/Engine/Source/Editor/WorldBrowser/Private/Tiles/WorldTileModel.cpp
+++ b/Engine/Source/Editor/WorldBrowser/Private/Tiles/WorldTileModel.cpp
@@ -672,8 +672,8 @@ void FWorldTileModel::OnParentPackageNamePropertyChanged()
 	if (GetLevelObject())
 	{
 		TSharedPtr<FLevelModel> NewParent = LevelCollectionModel.FindLevelModel(TileDetails->ParentPackageName);
-		// Assign to a root level if new parent is not found
-		if (!NewParent.IsValid()) 
+		// Assign to a root level if new parent is not found or we assigning to ourself 
+		if (!NewParent.IsValid() || NewParent.Get() == this) 
 		{
 			NewParent = static_cast<FWorldTileCollectionModel&>(LevelCollectionModel).GetWorldRootModel();
 		}
diff --git a/Engine/Source/Runtime/Engine/Private/WorldComposition.cpp b/Engine/Source/Runtime/Engine/Private/WorldComposition.cpp
--- a/Engine/Source/Runtime/Engine/Private/WorldComposition.cpp
+++ b/Engine/Source/Runtime/Engine/Private/WorldComposition.cpp
@@ -300,7 +300,7 @@ void UWorldComposition::CaclulateTilesAbsolutePositions()
 			{
 				UE_LOG(LogWorldComposition, Warning, TEXT("World composition tile (%s) has a cycled parent (%s)"), *Tile.PackageName.ToString(), *NextParentTileName.ToString());
 				NextParentTile = nullptr;
-				ParentTile->Info.ParentTilePackageName = TEXT("");
+				ParentTile->Info.ParentTilePackageName = FName(NAME_None).ToString();
 			}
 			
 			ParentTile = NextParentTile;
