Fixed selection border drawing regression in world composition

diff --git a/Engine/Source/Editor/WorldBrowser/Private/Tiles/SWorldComposition.cpp b/Engine/Source/Editor/WorldBrowser/Private/Tiles/SWorldComposition.cpp
--- a/Engine/Source/Editor/WorldBrowser/Private/Tiles/SWorldComposition.cpp
+++ b/Engine/Source/Editor/WorldBrowser/Private/Tiles/SWorldComposition.cpp
@@ -272,7 +272,7 @@ public:
 			}
 		}
 	}
-	
+
 	/**  SWidget interface */
 	virtual int32 OnPaint(const FPaintArgs& Args, const FGeometry& AllottedGeometry, const FSlateRect& MyClippingRect, FSlateWindowElementList& OutDrawElements, int32 LayerId, const FWidgetStyle& InWidgetStyle, bool bParentEnabled ) const override
 	{
@@ -314,8 +314,11 @@ public:
 			float ScreenSpaceSize = FLevelCollectionModel::EditableAxisLength()*GetZoomAmount()*2.f;
 			FVector2D PaintSize = FVector2D(ScreenSpaceSize, ScreenSpaceSize);
 			FVector2D PaintPosition = GraphCoordToPanelCoord(FVector2D::ZeroVector) - (PaintSize*0.5f);
-			FPaintGeometry EditableArea = AllottedGeometry.ToPaintGeometry(PaintPosition, PaintSize);
-			EditableArea.DrawScale = 0.2f;
+			float Scale = 0.2f; // Scale down drawing border
+			FSlateLayoutTransform LayoutTransform(Scale, AllottedGeometry.GetAccumulatedLayoutTransform().GetTranslation() + PaintPosition);
+			FSlateRenderTransform RenderTransform(Scale, AllottedGeometry.GetAccumulatedRenderTransform().GetTranslation() + PaintPosition);
+			FPaintGeometry EditableArea(LayoutTransform, RenderTransform, PaintSize/Scale);
+
 			FLinearColor PaintColor = FLinearColor::Yellow;
 			PaintColor.A = 0.4f;
 
diff --git a/Engine/Source/Editor/WorldBrowser/Private/Tiles/SWorldTileItem.cpp b/Engine/Source/Editor/WorldBrowser/Private/Tiles/SWorldTileItem.cpp
--- a/Engine/Source/Editor/WorldBrowser/Private/Tiles/SWorldTileItem.cpp
+++ b/Engine/Source/Editor/WorldBrowser/Private/Tiles/SWorldTileItem.cpp
@@ -378,9 +378,13 @@ int32 SWorldTileItem::OnPaint(const FPaintArgs& Args, const FGeometry& AllottedG
 		// Draw the node's selection.
 		if (bSelected)
 		{
-			FPaintGeometry SelectionGeometry = AllottedGeometry.ToInflatedPaintGeometry(FVector2D(4,4)/AllottedGeometry.Scale);
-			SelectionGeometry.DrawScale = 0.5;
-				
+			// Calculate selection box paint geometry 
+			const FVector2D InflateAmount = FVector2D(4, 4);
+			const float Scale = 0.5f; // Scale down image of the borders to make them thinner 
+			FSlateLayoutTransform LayoutTransform(Scale, AllottedGeometry.GetAccumulatedLayoutTransform().GetTranslation() - InflateAmount);
+			FSlateRenderTransform RenderTransform(Scale, AllottedGeometry.GetAccumulatedRenderTransform().GetTranslation() - InflateAmount);
+			FPaintGeometry SelectionGeometry(LayoutTransform, RenderTransform, (AllottedGeometry.GetLocalSize()*AllottedGeometry.Scale + InflateAmount*2)/Scale);
+										
 			FSlateDrawElement::MakeBox(
 				OutDrawElements,
 				LayerId + 1,
@@ -394,9 +398,6 @@ int32 SWorldTileItem::OnPaint(const FPaintArgs& Args, const FGeometry& AllottedG
 		//const auto& PreviewStreamingLevels = WorldModel->GetPreviewStreamingLevels();
 		//if (PreviewStreamingLevels.Find(LevelModel->TileDetails->PackageName) != INDEX_NONE)
 		//{
-		//	FPaintGeometry SelectionGeometry = AllottedGeometry.ToInflatedPaintGeometry(FVector2D(4,4)/AllottedGeometry.Scale);
-		//	SelectionGeometry.DrawScale = 0.5;
-
 		//	FSlateDrawElement::MakeBox(
 		//		OutDrawElements,
 		//		LayerId + 1,
