Fixed loading bar drawing in world composition

diff --git a/Engine/Source/Editor/WorldBrowser/Private/Tiles/SWorldTileItem.cpp b/Engine/Source/Editor/WorldBrowser/Private/Tiles/SWorldTileItem.cpp
--- a/Engine/Source/Editor/WorldBrowser/Private/Tiles/SWorldTileItem.cpp
+++ b/Engine/Source/Editor/WorldBrowser/Private/Tiles/SWorldTileItem.cpp
@@ -412,21 +412,24 @@ int32 SWorldTileItem::OnPaint(const FPaintArgs& Args, const FGeometry& AllottedG
 		// Draw progress bar if level is currently loading
 		if (TileModel->IsLoading())
 		{
-			const float ProgressBarAnimOffset = ProgressBarImage->ImageSize.X * CurveSequence.GetLerpLooping() / AllottedGeometry.Scale;
-			const float ProgressBarImageSize = ProgressBarImage->ImageSize.X / AllottedGeometry.Scale;
-			const float ProgressBarImageHeight = ProgressBarImage->ImageSize.Y / AllottedGeometry.Scale;
+			const float ProgressBarAnimOffset = ProgressBarImage->ImageSize.X * CurveSequence.GetLerpLooping();
+			const float ProgressBarImageSize = ProgressBarImage->ImageSize.X;
+			const float ProgressBarImageHeight = ProgressBarImage->ImageSize.Y;
+			const FVector2D ProggresBarOffset = FVector2D(ProgressBarAnimOffset - ProgressBarImageSize, 0);
+
+			FSlateLayoutTransform ProgressBarLayoutTransform(1.f, AllottedGeometry.GetAccumulatedLayoutTransform().GetTranslation() + ProggresBarOffset);
+			FSlateRenderTransform ProgressBarRenderTransform(1.f, AllottedGeometry.GetAccumulatedRenderTransform().GetTranslation() + ProggresBarOffset);
+			FPaintGeometry ProgressBarPaintGeometry(
+				ProgressBarLayoutTransform, 
+				ProgressBarRenderTransform, 
+				FVector2D(AllottedGeometry.Size.X*AllottedGeometry.Scale + ProgressBarImageSize, FMath::Min(AllottedGeometry.Size.Y*AllottedGeometry.Scale, ProgressBarImageHeight)));
+								
 			const FSlateRect ProgressBarClippingRect = AllottedGeometry.GetClippingRect().IntersectionWith(ClippingRect);
-		
-			FPaintGeometry LoadingBarPaintGeometry = AllottedGeometry.ToPaintGeometry(
-					FVector2D(ProgressBarAnimOffset - ProgressBarImageSize, 0 ),
-					FVector2D(AllottedGeometry.Size.X + ProgressBarImageSize, FMath::Min(AllottedGeometry.Size.Y, ProgressBarImageHeight))
-			);
-			LoadingBarPaintGeometry.DrawScale = 1.f;
 								
 			FSlateDrawElement::MakeBox(
 				OutDrawElements,
 				LayerId + 1,
-				LoadingBarPaintGeometry,
+				ProgressBarPaintGeometry,
 				ProgressBarImage,
 				ProgressBarClippingRect,
 				ESlateDrawEffect::None,
