Fixed preview of visible streaming levels in world composition

diff --git a/Engine/Source/Editor/WorldBrowser/Private/Tiles/SWorldComposition.cpp b/Engine/Source/Editor/WorldBrowser/Private/Tiles/SWorldComposition.cpp
--- a/Engine/Source/Editor/WorldBrowser/Private/Tiles/SWorldComposition.cpp
+++ b/Engine/Source/Editor/WorldBrowser/Private/Tiles/SWorldComposition.cpp
@@ -204,7 +204,7 @@ public:
 		const bool bShowPotentiallyVisibleLevels = FSlateApplication::Get().GetModifierKeys().IsAltDown() && 
 													AllottedGeometry.IsUnderLocation(CursorPosition);
 	
-		//WorldModel->UpdateStreamingPreview(WorldMouseLocation, bShowPotentiallyVisibleLevels);
+		WorldModel->UpdateStreamingPreview(WorldMouseLocation, bShowPotentiallyVisibleLevels);
 			
 		// deffered scroll and zooming requests
 		if (bHasScrollToRequest || bHasScrollByRequest)
diff --git a/Engine/Source/Editor/WorldBrowser/Private/Tiles/SWorldTileItem.cpp b/Engine/Source/Editor/WorldBrowser/Private/Tiles/SWorldTileItem.cpp
--- a/Engine/Source/Editor/WorldBrowser/Private/Tiles/SWorldTileItem.cpp
+++ b/Engine/Source/Editor/WorldBrowser/Private/Tiles/SWorldTileItem.cpp
@@ -374,9 +374,10 @@ int32 SWorldTileItem::OnPaint(const FPaintArgs& Args, const FGeometry& AllottedG
 		LayerId = SNodePanel::SNode::OnPaint(Args, AllottedGeometry, ClippingRect, OutDrawElements, LayerId, InWidgetStyle, bParentEnabled);
 		
 		const bool bSelected = (IsItemSelected() || bAffectedByMarquee);
+		const bool bHighlighted = (WorldModel->GetPreviewStreamingLevels().Find(TileModel->TileDetails->PackageName) != nullptr);
 
-		// Draw the node's selection.
-		if (bSelected)
+		// Draw the node's selection/highlight.
+		if (bSelected || bHighlighted)
 		{
 			// Calculate selection box paint geometry 
 			const FVector2D InflateAmount = FVector2D(4, 4);
@@ -389,26 +390,13 @@ int32 SWorldTileItem::OnPaint(const FPaintArgs& Args, const FGeometry& AllottedG
 				OutDrawElements,
 				LayerId + 1,
 				SelectionGeometry,
-				GetShadowBrush(bSelected),
-				ClippingRect
+				GetShadowBrush(bSelected || bHighlighted),
+				ClippingRect,
+				ESlateDrawEffect::None,
+				bHighlighted ? FLinearColor::Green : FLinearColor::White
 				);
 		}
 
-		// Highlight potentially visible streaming levels from current preview point
-		//const auto& PreviewStreamingLevels = WorldModel->GetPreviewStreamingLevels();
-		//if (PreviewStreamingLevels.Find(LevelModel->TileDetails->PackageName) != INDEX_NONE)
-		//{
-		//	FSlateDrawElement::MakeBox(
-		//		OutDrawElements,
-		//		LayerId + 1,
-		//		SelectionGeometry,
-		//		GetShadowBrush(true),
-		//		ClippingRect,
-		//		ESlateDrawEffect::None,
-		//		FLinearColor::Green
-		//		);
-		//}
-
 		// Draw progress bar if level is currently loading
 		if (TileModel->IsLoading())
 		{
diff --git a/Engine/Source/Editor/WorldBrowser/Private/Tiles/WorldTileCollectionModel.cpp b/Engine/Source/Editor/WorldBrowser/Private/Tiles/WorldTileCollectionModel.cpp
--- a/Engine/Source/Editor/WorldBrowser/Private/Tiles/WorldTileCollectionModel.cpp
+++ b/Engine/Source/Editor/WorldBrowser/Private/Tiles/WorldTileCollectionModel.cpp
@@ -856,7 +856,7 @@ void FWorldTileCollectionModel::DeselectLevels(const FWorldTileLayer& InLayer)
 	}
 }
 
-TArray<FName>& FWorldTileCollectionModel::GetPreviewStreamingLevels()
+const TSet<FName>& FWorldTileCollectionModel::GetPreviewStreamingLevels() const
 {
 	return PreviewVisibleTiles;
 }
diff --git a/Engine/Source/Editor/WorldBrowser/Private/Tiles/WorldTileCollectionModel.h b/Engine/Source/Editor/WorldBrowser/Private/Tiles/WorldTileCollectionModel.h
--- a/Engine/Source/Editor/WorldBrowser/Private/Tiles/WorldTileCollectionModel.h
+++ b/Engine/Source/Editor/WorldBrowser/Private/Tiles/WorldTileCollectionModel.h
@@ -111,8 +111,8 @@ public:
 	/**	Notification that "view point" for streaming levels visibility preview was changed */
 	void UpdateStreamingPreview(FVector2D InPreviewLocation, bool bEnabled);
 
-	/** Returns list of visible streaming levels for curent preview location */
-	TArray<FName>& GetPreviewStreamingLevels();
+	/** Returns list of visible streaming levels for current preview location */
+	const TSet<FName>& GetPreviewStreamingLevels() const;
 		
 	/** Calculates snapped moving delta based on specified landscape tile */
 	FVector2D SnapTranslationDeltaLandscape(const TSharedPtr<FWorldTileModel>& LandscapeTile, 
@@ -250,7 +250,7 @@ private:
 	FLevelModelList						StaticTileList;	
 
 	/** Cached streaming tiles which are potentially visible from specified view point*/
-	TArray<FName>						PreviewVisibleTiles;
+	TSet<FName>							PreviewVisibleTiles;
 	
 	/** View point location for calculating potentially visible streaming tiles*/
 	FVector								PreviewLocation;
