Improved Streaming Level customization to protect against entering invalid streaming volumes in the EditorStreamingVolumes list.

diff --git a/Engine/Source/Editor/WorldBrowser/Private/StreamingLevels/StreamingLevelCustomization.cpp b/Engine/Source/Editor/WorldBrowser/Private/StreamingLevels/StreamingLevelCustomization.cpp
--- a/Engine/Source/Editor/WorldBrowser/Private/StreamingLevels/StreamingLevelCustomization.cpp
+++ b/Engine/Source/Editor/WorldBrowser/Private/StreamingLevels/StreamingLevelCustomization.cpp
@@ -4,6 +4,7 @@
 #include "Editor/PropertyEditor/Public/IDetailsView.h"
 #include "Editor/PropertyEditor/Public/PropertyEditing.h"
 #include "Editor/PropertyEditor/Public/PropertyEditorModule.h"
+#include "Editor/PropertyEditor/Public/PropertyCustomizationHelpers.h"
 #include "SVectorInputBox.h"
 #include "SRotatorInputBox.h"
 
@@ -113,6 +114,95 @@ void FStreamingLevelCustomization::CustomizeDetails(IDetailLayoutBuilder& Detail
 				.ContentPadding(1)
 			]
 		];
+
+	TSharedRef<IPropertyHandle> EditorStreamingVolumesProperty = DetailLayoutBuilder.GetProperty(GET_MEMBER_NAME_CHECKED(ULevelStreaming, EditorStreamingVolumes));
+	const bool bGenerateHeader = true;
+	const bool bDisplayResetToDefault = false;
+	TSharedRef<FDetailArrayBuilder> EditorStreamingVolumesBuilder = MakeShareable(new FDetailArrayBuilder(EditorStreamingVolumesProperty, bGenerateHeader, bDisplayResetToDefault));
+	EditorStreamingVolumesBuilder->OnGenerateArrayElementWidget(FOnGenerateArrayElementWidget::CreateSP(this, &FStreamingLevelCustomization::OnGenerateElementForEditorStreamingVolume));
+
+	const bool bForAdvanced = false;
+	LevelStreamingCategory.AddCustomBuilder(EditorStreamingVolumesBuilder, bForAdvanced);
+}
+
+void FStreamingLevelCustomization::OnGenerateElementForEditorStreamingVolume(TSharedRef<IPropertyHandle> ElementProperty, int32 ElementIndex, IDetailChildrenBuilder& ChildrenBuilder)
+{
+	IDetailPropertyRow& PropertyRow = ChildrenBuilder.AddChildProperty(ElementProperty);
+	TSharedPtr<SWidget> NameWidget;
+	TSharedPtr<SWidget> ValueWidget;
+	FDetailWidgetRow Row;
+	PropertyRow.GetDefaultWidgets(NameWidget, ValueWidget, Row);
+
+	PropertyRow.CustomWidget()
+	.NameContent()
+	.MinDesiredWidth(Row.NameWidget.MinWidth)
+	.MaxDesiredWidth(Row.NameWidget.MaxWidth)
+	[
+		NameWidget.ToSharedRef()
+	]
+	.ValueContent()
+	.MinDesiredWidth(Row.ValueWidget.MinWidth)
+	.MaxDesiredWidth(Row.ValueWidget.MaxWidth)
+	[
+		SNew(SObjectPropertyEntryBox)
+		.PropertyHandle(ElementProperty)
+		.AllowedClass(ALevelStreamingVolume::StaticClass())
+		.OnShouldSetAsset(this, &FStreamingLevelCustomization::OnShouldSetEditorStreamingVolume, ElementProperty)
+	];
+}
+
+bool FStreamingLevelCustomization::OnShouldSetEditorStreamingVolume(const FAssetData& AssetData, TSharedRef<IPropertyHandle> ElementProperty) const
+{
+	ALevelStreamingVolume* Volume = Cast<ALevelStreamingVolume>(AssetData.GetAsset());
+
+	if (Volume != nullptr)
+	{
+		// Check if there are any duplicates
+		bool bIsUnique = true;
+		TSharedPtr<IPropertyHandle> ParentProperty = ElementProperty->GetParentHandle();
+		TSharedPtr<IPropertyHandleArray> ParentPropertyAsArray = ParentProperty.IsValid() ? ParentProperty->AsArray() : nullptr;
+		if (ParentPropertyAsArray.IsValid())
+		{
+			int32 Index = ElementProperty->GetIndexInArray();
+			check(Index != INDEX_NONE);
+
+			uint32 NumItems = 0;
+			ensure(ParentPropertyAsArray->GetNumElements(NumItems) == FPropertyAccess::Success);
+
+			for (uint32 ElementIndex = 0; ElementIndex < NumItems; ElementIndex++)
+			{
+				if (ElementIndex != Index)
+				{
+					TSharedRef<IPropertyHandle> ElementToCompare = ParentPropertyAsArray->GetElement(ElementIndex);
+					UObject* ElementValue = nullptr;
+					ensure(ElementToCompare->GetValue(ElementValue) == FPropertyAccess::Success);
+					if (ElementValue == Volume)
+					{
+						FMessageDialog::Open(
+							EAppMsgType::Ok,
+							LOCTEXT("DuplicateVolume", "This volume is already in the list."));
+
+						bIsUnique = false;
+						break;
+					}
+				}
+			}
+		}
+
+		// Check that the volume is in the persistent level
+		bool bIsInPersistentLevel = Volume->IsInPersistentLevel();
+
+		if (!bIsInPersistentLevel)
+		{
+			FMessageDialog::Open(
+				EAppMsgType::Ok,
+				LOCTEXT("VolumeMustBeInPersistentLevel", "Cannot add a Level Streaming Volume which is not in the persistent level."));
+		}
+
+		return bIsInPersistentLevel && bIsUnique;
+	}
+
+	return false;
 }
 
 void FStreamingLevelCustomization::OnSetLevelPosition( float NewValue, ETextCommit::Type CommitInfo, int32 Axis )
diff --git a/Engine/Source/Editor/WorldBrowser/Private/StreamingLevels/StreamingLevelCustomization.h b/Engine/Source/Editor/WorldBrowser/Private/StreamingLevels/StreamingLevelCustomization.h
--- a/Engine/Source/Editor/WorldBrowser/Private/StreamingLevels/StreamingLevelCustomization.h
+++ b/Engine/Source/Editor/WorldBrowser/Private/StreamingLevels/StreamingLevelCustomization.h
@@ -71,6 +71,12 @@ private:
 	/** Handle for edit level transform button */
 	FReply OnEditLevelClicked();
 
+	/** Generate custom widget for the EditorStreamingVolumes array customization */
+	void OnGenerateElementForEditorStreamingVolume(TSharedRef<IPropertyHandle> ElementProperty, int32 ElementIndex, IDetailChildrenBuilder& ChildrenBuilder);
+
+	/** Called to determine whether the selected streaming volume can be set or not */
+	bool OnShouldSetEditorStreamingVolume(const FAssetData& AssetData, TSharedRef<IPropertyHandle> ElementProperty) const;
+
 private:
 	TWeakPtr<FStreamingLevelCollectionModel>	WorldModel;
 	TSharedPtr<class IPropertyHandle>			LevelPositionProperty;
