Tiled landscape import: resolution rounding and spelling fixes

diff --git a/Engine/Source/Editor/WorldBrowser/Private/Tiles/STiledLandscapeImportDlg.cpp b/Engine/Source/Editor/WorldBrowser/Private/Tiles/STiledLandscapeImportDlg.cpp
--- a/Engine/Source/Editor/WorldBrowser/Private/Tiles/STiledLandscapeImportDlg.cpp
+++ b/Engine/Source/Editor/WorldBrowser/Private/Tiles/STiledLandscapeImportDlg.cpp
@@ -294,7 +294,7 @@ FText STiledLandcapeImportDlg::GetTileConfigurationText() const
 
 	if (ImportSettings.SectionsPerComponent <= 0)
 	{
-		return LOCTEXT("TiledLandscapeImport_InvalidTileResolutuinText", "Selected tiles have unsupported resolutuion");
+		return LOCTEXT("TiledLandscapeImport_InvalidTileResolutionText", "Selected tiles have unsupported resolution");
 	}
 	
 	return GenerateConfigurationText(ImportSettings.ComponentsNum, ImportSettings.SectionsPerComponent, ImportSettings.QuadsPerSection);
@@ -487,7 +487,7 @@ FReply STiledLandcapeImportDlg::OnClickedSelectHeightmapTiles()
 				
 				if (bValidTiles)
 				{
-					ImportSettings.TileResolution = FMath::Sqrt(TargetFileSize/2.0);
+					ImportSettings.TileResolution = FMath::RoundToInt(FMath::Sqrt(TargetFileSize/2.0f));
 					SetPossibleConfigurationsForResolution(ImportSettings.TileResolution);
 				}
 			}
@@ -576,14 +576,14 @@ void STiledLandcapeImportDlg::OnLandscapeMaterialChanged(const FAssetData& Asset
 void STiledLandcapeImportDlg::SetPossibleConfigurationsForResolution(int32 TargetResolutuion)
 {
 	int32 Idx = AllConfigurations.IndexOfByPredicate([&](const FTileImportConfiguration& A){
-		return TargetResolutuion == A.Resolutuion;
+		return TargetResolutuion == A.Resolution;
 	});
 
 	ActiveConfigurations.Empty();
 	ImportSettings.ComponentsNum = 0; // Set invalid options
 
 	// AllConfigurations is sorted by resolution
-	while(AllConfigurations.IsValidIndex(Idx) && AllConfigurations[Idx].Resolutuion == TargetResolutuion)
+	while(AllConfigurations.IsValidIndex(Idx) && AllConfigurations[Idx].Resolution == TargetResolutuion)
 	{
 		TSharedPtr<FTileImportConfiguration> TileConfig = MakeShareable(new FTileImportConfiguration(AllConfigurations[Idx++]));
 		ActiveConfigurations.Add(TileConfig);
@@ -611,7 +611,7 @@ void STiledLandcapeImportDlg::GenerateAllPossibleTileConfigurations()
 				Entry.NumComponents				= ComponentsNum;
 				Entry.NumSectionsPerComponent	= SectionsPerComponent;
 				Entry.NumQuadsPerSection		= (1 << QuadsPerSection) - 1;
-				Entry.Resolutuion				= CalcLandscapeSquareResolution(Entry.NumComponents, Entry.NumSectionsPerComponent, Entry.NumQuadsPerSection);
+				Entry.Resolution				= CalcLandscapeSquareResolution(Entry.NumComponents, Entry.NumSectionsPerComponent, Entry.NumQuadsPerSection);
 
 				AllConfigurations.Add(Entry);
 			}
@@ -620,11 +620,11 @@ void STiledLandcapeImportDlg::GenerateAllPossibleTileConfigurations()
 	
 	// Sort by resolution
 	AllConfigurations.Sort([](const FTileImportConfiguration& A, const FTileImportConfiguration& B){
-		if (A.Resolutuion == B.Resolutuion)
+		if (A.Resolution == B.Resolution)
 		{
 			return A.NumComponents < B.NumComponents;
 		}
-		return A.Resolutuion < B.Resolutuion;
+		return A.Resolution < B.Resolution;
 	});
 }
 
diff --git a/Engine/Source/Editor/WorldBrowser/Private/Tiles/STiledLandscapeImportDlg.h b/Engine/Source/Editor/WorldBrowser/Private/Tiles/STiledLandscapeImportDlg.h
--- a/Engine/Source/Editor/WorldBrowser/Private/Tiles/STiledLandscapeImportDlg.h
+++ b/Engine/Source/Editor/WorldBrowser/Private/Tiles/STiledLandscapeImportDlg.h
@@ -11,7 +11,7 @@ public:
 	/** */		
 	struct FTileImportConfiguration
 	{
-		int32 Resolutuion;
+		int32 Resolution;
 		
 		int32 NumComponents;
 		int32 NumSectionsPerComponent;
