Fixed crash on LOD level generation in World Composition

diff --git a/Engine/Source/Editor/WorldBrowser/Private/Tiles/WorldTileCollectionModel.cpp b/Engine/Source/Editor/WorldBrowser/Private/Tiles/WorldTileCollectionModel.cpp
--- a/Engine/Source/Editor/WorldBrowser/Private/Tiles/WorldTileCollectionModel.cpp
+++ b/Engine/Source/Editor/WorldBrowser/Private/Tiles/WorldTileCollectionModel.cpp
@@ -1974,8 +1974,15 @@ bool FWorldTileCollectionModel::GenerateLODLevels(FLevelModelList InLevelList, i
 		// Create new level and spawn generated assets in it
 		if (AssetsToSpawn.Num())
 		{
+			UWorld* LODWorld = UWorld::FindWorldInPackage(LODPackage);
+			if (LODWorld)
+			{
+				LODWorld->ClearFlags(RF_Standalone);
+				LODWorld->DestroyWorld(false);
+			}
+
 			// Create a new world
-			UWorld* LODWorld = UWorld::CreateWorld(EWorldType::None, false, FPackageName::GetShortFName(LODPackage->GetFName()), LODPackage);
+			LODWorld = UWorld::CreateWorld(EWorldType::None, false, FPackageName::GetShortFName(LODPackage->GetFName()), LODPackage);
 			LODWorld->SetFlags(RF_Standalone);
 
 			for (int32 AssetIdx = 0; AssetIdx < AssetsToSpawn.Num(); ++AssetIdx)
