Merging back to main from GDC2015 branch.
-Prevent landscape collision from checking the DDC when we know we are generating new collision data for the first time.
-Do not recreate landscape collision objects when actor is moved
-Added error messages on tiled landscape import dilaog to explain what is wrong with source data
-Added png format support when importing tiled weighmaps
-Show material instances in tiled landscape import dialog

diff --git a/Engine/Source/Editor/WorldBrowser/Private/LevelCollectionModel.h b/Engine/Source/Editor/WorldBrowser/Private/LevelCollectionModel.h
--- a/Engine/Source/Editor/WorldBrowser/Private/LevelCollectionModel.h
+++ b/Engine/Source/Editor/WorldBrowser/Private/LevelCollectionModel.h
@@ -569,7 +569,7 @@ struct FTiledLandscapeImportSettings
 		, QuadsPerSection(63)
 		, SectionsPerComponent(1)
 		, TilesCoordinatesOffset(0,0)
-		, TileResolution(1009)
+		, SizeX(1009)
 	{}
 	
 	FVector				Scale3D;
@@ -580,7 +580,7 @@ struct FTiledLandscapeImportSettings
 	TArray<FString>		HeightmapFileList;
 	TArray<FIntPoint>	TileCoordinates;
 	FIntPoint			TilesCoordinatesOffset;	
-	int32				TileResolution;	
+	int32				SizeX;
 
 
 	TWeakObjectPtr<UMaterialInterface>	LandscapeMaterial;
diff --git a/Engine/Source/Editor/WorldBrowser/Private/Tiles/STiledLandscapeImportDlg.cpp b/Engine/Source/Editor/WorldBrowser/Private/Tiles/STiledLandscapeImportDlg.cpp
--- a/Engine/Source/Editor/WorldBrowser/Private/Tiles/STiledLandscapeImportDlg.cpp
+++ b/Engine/Source/Editor/WorldBrowser/Private/Tiles/STiledLandscapeImportDlg.cpp
@@ -199,12 +199,13 @@ void STiledLandcapeImportDlg::Construct(const FArguments& InArgs, TSharedPtr<SWi
 			// Import summary
 			+SVerticalBox::Slot()
 			.AutoHeight()
-			.HAlign(HAlign_Right)
+			.HAlign(HAlign_Left)
 			.VAlign(VAlign_Center)
 			.Padding(FEditorStyle::GetMargin("StandardDialog.ContentPadding"))
 			[
 				SNew(STextBlock)
 				.Text(this, &STiledLandcapeImportDlg::GetImportSummaryText)
+				.WrapTextAt(600.0f)
 			]
 
 			// Import, Cancel
@@ -241,7 +242,7 @@ void STiledLandcapeImportDlg::Construct(const FArguments& InArgs, TSharedPtr<SWi
 
 
 	GenerateAllPossibleTileConfigurations();
-	SetPossibleConfigurationsForResolution(0);
+	SetPossibleConfigurationsForFileSize(0);
 }
 
 TSharedRef<SWidget> STiledLandcapeImportDlg::CreateLandscapeMaterialPicker()
@@ -249,13 +250,14 @@ TSharedRef<SWidget> STiledLandcapeImportDlg::CreateLandscapeMaterialPicker()
 	FContentBrowserModule& ContentBrowserModule = FModuleManager::Get().LoadModuleChecked<FContentBrowserModule>(TEXT("ContentBrowser"));
 
 	FAssetPickerConfig AssetPickerConfig;
+	AssetPickerConfig.Filter.bRecursiveClasses = true;
 	AssetPickerConfig.Filter.ClassNames.Add(UMaterial::StaticClass()->GetFName());
 	AssetPickerConfig.Filter.ClassNames.Add(UMaterialInstance::StaticClass()->GetFName());
 	AssetPickerConfig.OnAssetSelected = FOnAssetSelected::CreateSP(this, &STiledLandcapeImportDlg::OnLandscapeMaterialChanged);
 	AssetPickerConfig.InitialAssetViewType = EAssetViewType::List;
 	AssetPickerConfig.bAllowNullSelection = true;
 	AssetPickerConfig.ThumbnailScale = 0.0f;
-	
+		
 	return SNew(SBox)
 		.WidthOverride(250)
 		.HeightOverride(300)
@@ -424,6 +426,7 @@ void STiledLandcapeImportDlg::OnSetImportConfiguration(TSharedPtr<FTileImportCon
 		ImportSettings.ComponentsNum = InTileConfig->NumComponents;
 		ImportSettings.QuadsPerSection = InTileConfig->NumQuadsPerSection;
 		ImportSettings.SectionsPerComponent = InTileConfig->NumSectionsPerComponent;
+		ImportSettings.SizeX = InTileConfig->SizeX;
 	}
 	else
 	{
@@ -438,7 +441,7 @@ FReply STiledLandcapeImportDlg::OnClickedSelectHeightmapTiles()
 	ImportSettings.HeightmapFileList.Empty();
 	ImportSettings.TileCoordinates.Empty();
 	
-	SetPossibleConfigurationsForResolution(0);
+	SetPossibleConfigurationsForFileSize(0);
 	
 	IDesktopPlatform* DesktopPlatform = FDesktopPlatformModule::Get();
 	if (DesktopPlatform)
@@ -447,7 +450,7 @@ FReply STiledLandcapeImportDlg::OnClickedSelectHeightmapTiles()
 		{
 			void* ParentWindowWindowHandle = ParentWindow->GetNativeWindow()->GetOSWindowHandle();
 
-			const FString ImportFileTypes = TEXT("Raw heightmap tiles (*.r16)|*.r16");
+			const FString ImportFileTypes = TEXT("Heightmap tiles (*.r16)|*.r16");
 			bool bOpened = DesktopPlatform->OpenFileDialog(
 								ParentWindowWindowHandle,
 								LOCTEXT("SelectHeightmapTiles", "Select heightmap tiles").ToString(),
@@ -467,8 +470,22 @@ FReply STiledLandcapeImportDlg::OnClickedSelectHeightmapTiles()
 				for (const FString& Filename : ImportSettings.HeightmapFileList)
 				{
 					int64 FileSize = FileManager.FileSize(*Filename);
+					
+					if ((FileSize & 1) != 0) // we expect 2 bytes per sample
+					{
+						FFormatNamedArguments Arguments;
+						Arguments.Add(TEXT("FileName"), FText::FromString(Filename));
+						StatusMessage = FText::Format(LOCTEXT("TiledLandscapeImport_HeightmapTileInvalidFormat", "File ({FileName}) should have raw data with 2 bytes per sample."), Arguments);
+						bValidTiles = false;
+					}
+					
 					if (FileSize != TargetFileSize)
 					{
+						FFormatNamedArguments Arguments;
+						Arguments.Add(TEXT("FileName"), FText::FromString(Filename));
+						Arguments.Add(TEXT("FileSize"), FileSize);
+						Arguments.Add(TEXT("TargetFileSize"), TargetFileSize);
+						StatusMessage = FText::Format(LOCTEXT("TiledLandscapeImport_HeightmapTileSizeMismatch", "File ({FileName}) size ({FileSize}) should match other tiles file size ({TargetFileSize})."), Arguments);
 						bValidTiles = false;
 						break;
 					}
@@ -476,19 +493,27 @@ FReply STiledLandcapeImportDlg::OnClickedSelectHeightmapTiles()
 					FIntPoint TileCoordinate = ExtractTileCoordinates(FPaths::GetBaseFilename(Filename));
 					if (TileCoordinate.GetMin() < 0)
 					{
+						FFormatNamedArguments Arguments;
+						Arguments.Add(TEXT("FileName"), FText::FromString(Filename));
+						StatusMessage = FText::Format(LOCTEXT("TiledLandscapeImport_HeightmapTileInvalidName", "File name ({FileName}) should match pattern: <name>_X<number>_Y<number>."), Arguments);
 						bValidTiles = false;
 						break;
 					}
 
 					TotalLandscapeRect.Include(TileCoordinate);
 					ImportSettings.TileCoordinates.Add(TileCoordinate);
-					
 				}
 				
 				if (bValidTiles)
 				{
-					ImportSettings.TileResolution = FMath::RoundToInt(FMath::Sqrt(TargetFileSize/2.0f));
-					SetPossibleConfigurationsForResolution(ImportSettings.TileResolution);
+					if (SetPossibleConfigurationsForFileSize(TargetFileSize) < 1)
+					{
+						int64 NumSamples = TargetFileSize/2; // 2 bytes per sample
+						FFormatNamedArguments Arguments;
+						Arguments.Add(TEXT("NumSamples"), NumSamples);
+						StatusMessage = FText::Format(LOCTEXT("TiledLandscapeImport_HeightmapTileInvalidSize", "No landscape configuration found for ({NumSamples}) samples."), Arguments);
+						bValidTiles = false;
+					}
 				}
 			}
 		}
@@ -510,7 +535,7 @@ FReply STiledLandcapeImportDlg::OnClickedSelectWeighmapTiles(TSharedPtr<FTiledLa
 
 			TArray<FString> WeightmapFilesList;
 
-			const FString ImportFileTypes = TEXT("Raw weigtmap tiles (*.raw)|*.raw");
+			const FString ImportFileTypes = TEXT("Weightmap tiles (*.raw, *.png)|*.raw;*.png");
 			bool bOpened = DesktopPlatform->OpenFileDialog(
 								ParentWindowWindowHandle,
 								LOCTEXT("SelectHeightmapTiles", "Select weightmap tiles").ToString(),
@@ -573,17 +598,18 @@ void STiledLandcapeImportDlg::OnLandscapeMaterialChanged(const FAssetData& Asset
 	UpdateLandscapeLayerList();
 }
 
-void STiledLandcapeImportDlg::SetPossibleConfigurationsForResolution(int32 TargetResolutuion)
+int32 STiledLandcapeImportDlg::SetPossibleConfigurationsForFileSize(int64 TargetFileSize)
 {
 	int32 Idx = AllConfigurations.IndexOfByPredicate([&](const FTileImportConfiguration& A){
-		return TargetResolutuion == A.Resolution;
+		return TargetFileSize == A.ImportFileSize;
 	});
 
 	ActiveConfigurations.Empty();
 	ImportSettings.ComponentsNum = 0; // Set invalid options
+	StatusMessage = FText();
 
-	// AllConfigurations is sorted by resolution
-	while(AllConfigurations.IsValidIndex(Idx) && AllConfigurations[Idx].Resolution == TargetResolutuion)
+	// AllConfigurations - sorted by resolution
+	while(AllConfigurations.IsValidIndex(Idx) && AllConfigurations[Idx].ImportFileSize == TargetFileSize)
 	{
 		TSharedPtr<FTileImportConfiguration> TileConfig = MakeShareable(new FTileImportConfiguration(AllConfigurations[Idx++]));
 		ActiveConfigurations.Add(TileConfig);
@@ -596,6 +622,8 @@ void STiledLandcapeImportDlg::SetPossibleConfigurationsForResolution(int32 Targe
 	{
 		TileConfigurationComboBox->SetSelectedItem(ActiveConfigurations[0]);
 	}
+
+	return ActiveConfigurations.Num();
 }
 
 void STiledLandcapeImportDlg::GenerateAllPossibleTileConfigurations()
@@ -611,7 +639,9 @@ void STiledLandcapeImportDlg::GenerateAllPossibleTileConfigurations()
 				Entry.NumComponents				= ComponentsNum;
 				Entry.NumSectionsPerComponent	= SectionsPerComponent;
 				Entry.NumQuadsPerSection		= (1 << QuadsPerSection) - 1;
-				Entry.Resolution				= CalcLandscapeSquareResolution(Entry.NumComponents, Entry.NumSectionsPerComponent, Entry.NumQuadsPerSection);
+				Entry.SizeX = CalcLandscapeSquareResolution(Entry.NumComponents, Entry.NumSectionsPerComponent, Entry.NumQuadsPerSection);
+				// Calculate expected file size for this setup
+				Entry.ImportFileSize = FMath::Square(Entry.SizeX)*2; // 2 bytes per sample 
 
 				AllConfigurations.Add(Entry);
 			}
@@ -620,36 +650,36 @@ void STiledLandcapeImportDlg::GenerateAllPossibleTileConfigurations()
 	
 	// Sort by resolution
 	AllConfigurations.Sort([](const FTileImportConfiguration& A, const FTileImportConfiguration& B){
-		if (A.Resolution == B.Resolution)
+		if (A.SizeX == B.SizeX)
 		{
 			return A.NumComponents < B.NumComponents;
 		}
-		return A.Resolution < B.Resolution;
+		return A.SizeX < B.SizeX;
 	});
 }
 
 FText STiledLandcapeImportDlg::GetImportSummaryText() const
 {
-	if (ImportSettings.HeightmapFileList.Num() == 0 || ImportSettings.ComponentsNum <= 0)
+	if (ImportSettings.HeightmapFileList.Num() && ImportSettings.ComponentsNum > 0)
 	{
-		return FText();
-	}
+		// Tile information(num, resolution)
+		const FString TilesSummary = FString::Printf(TEXT("%d - %dx%d"), ImportSettings.HeightmapFileList.Num(), ImportSettings.SizeX, ImportSettings.SizeX);
 	
-	// Tile information(num, resolution)
-	const FString TilesSummary = FString::Printf(TEXT("%d - %dx%d"), ImportSettings.HeightmapFileList.Num(), ImportSettings.TileResolution, ImportSettings.TileResolution);
+		// Total landscape size(NxN km)
+		int32 WidthInTilesX = TotalLandscapeRect.Width() + 1;
+		int32 WidthInTilesY = TotalLandscapeRect.Height() + 1;
+		float WidthX = 0.00001f*ImportSettings.Scale3D.X*WidthInTilesX*ImportSettings.SizeX;
+		float WidthY = 0.00001f*ImportSettings.Scale3D.Y*WidthInTilesY*ImportSettings.SizeX;
+		const FString LandscapeSummary = FString::Printf(TEXT("%.3fx%.3f"), WidthX, WidthY);
 	
-	// Total landscape size(NxN km)
-	int32 WidthInTilesX = TotalLandscapeRect.Width() + 1;
-	int32 WidthInTilesY = TotalLandscapeRect.Height() + 1;
-	float WidthX = 0.00001f*ImportSettings.Scale3D.X*WidthInTilesX*ImportSettings.TileResolution;
-	float WidthY = 0.00001f*ImportSettings.Scale3D.Y*WidthInTilesY*ImportSettings.TileResolution;
-	const FString LandscapeSummary = FString::Printf(TEXT("%.3fx%.3f"), WidthX, WidthY);
+		StatusMessage = FText::Format(
+			LOCTEXT("TiledLandscapeImport_SummaryText", "{0} tiles, {1}km landscape"), 
+			FText::FromString(TilesSummary),
+			FText::FromString(LandscapeSummary)
+			);
+	}
 	
-	return FText::Format(
-		LOCTEXT("TiledLandscapeImport_SummaryText", "{0} tiles, {1}km landscape"), 
-		FText::FromString(TilesSummary),
-		FText::FromString(LandscapeSummary)
-		);
+	return StatusMessage;
 }
 
 FText STiledLandcapeImportDlg::GetWeightmapCountText(TSharedPtr<FTiledLandscapeImportSettings::LandscapeLayerSettings> InLayerData) const
diff --git a/Engine/Source/Editor/WorldBrowser/Private/Tiles/STiledLandscapeImportDlg.h b/Engine/Source/Editor/WorldBrowser/Private/Tiles/STiledLandscapeImportDlg.h
--- a/Engine/Source/Editor/WorldBrowser/Private/Tiles/STiledLandscapeImportDlg.h
+++ b/Engine/Source/Editor/WorldBrowser/Private/Tiles/STiledLandscapeImportDlg.h
@@ -11,8 +11,9 @@ public:
 	/** */		
 	struct FTileImportConfiguration
 	{
-		int32 Resolution;
+		int64 ImportFileSize;
 		
+		int32 SizeX;
 		int32 NumComponents;
 		int32 NumSectionsPerComponent;
 		int32 NumQuadsPerSection;
@@ -76,7 +77,7 @@ private:
 	void OnLayerBlendStateChanged(ESlateCheckBoxState::Type NewState, TSharedPtr<FTiledLandscapeImportSettings::LandscapeLayerSettings> InLayerData);
 
 	/** */
-	void SetPossibleConfigurationsForResolution(int32 TargetResolutuion);
+	int32 SetPossibleConfigurationsForFileSize(int64 InFileSize);
 	
 	/** */
 	void GenerateAllPossibleTileConfigurations();
@@ -90,6 +91,9 @@ private:
 private:
 	/** */
 	bool bShouldImport;
+
+	/** */
+	mutable FText StatusMessage;
 	
 	/** */
 	TSharedPtr<SWindow> ParentWindow;
diff --git a/Engine/Source/Editor/WorldBrowser/Private/Tiles/WorldTileCollectionModel.cpp b/Engine/Source/Editor/WorldBrowser/Private/Tiles/WorldTileCollectionModel.cpp
--- a/Engine/Source/Editor/WorldBrowser/Private/Tiles/WorldTileCollectionModel.cpp
+++ b/Engine/Source/Editor/WorldBrowser/Private/Tiles/WorldTileCollectionModel.cpp
@@ -11,6 +11,7 @@
 #include "MeshUtilities.h"
 #include "RawMesh.h"
 #include "LandscapeEdMode.h"
+#include "ImageWrapper.h"
 
 #include "WorldTileDetails.h"
 #include "WorldTileDetailsCustomization.h"
@@ -1389,23 +1390,44 @@ void FWorldTileCollectionModel::AddLandscapeProxy_Executed(FWorldTileModel::EWor
 }
 
 template<typename DataType>
-static bool ReadRawFile(TArray<DataType>& Result, const TCHAR* Filename, uint32 Flags = 0)
+static bool ReadRawFile(TArray<DataType>& Result, const FString& Filename, uint32 Flags = 0)
 {
-	auto Reader = TScopedPointer<FArchive>(IFileManager::Get().CreateFileReader(Filename, Flags));
+	auto Reader = TScopedPointer<FArchive>(IFileManager::Get().CreateFileReader(*Filename, Flags));
 	if (!Reader)
 	{
 		if (!(Flags & FILEREAD_Silent))
 		{
-			UE_LOG(LogStreaming, Warning, TEXT("Failed to read file '%s' error."), Filename);
+			UE_LOG(LogStreaming, Warning, TEXT("Failed to read file '%s' error."), *Filename);
 		}
 		return 0;
 	}
-	Result.Reset();
-	Result.AddUninitialized(Reader->TotalSize()/Result.GetTypeSize());
+	Result.Init(Reader->TotalSize()/Result.GetTypeSize());
 	Reader->Serialize(Result.GetData(), Result.Num()*Result.GetTypeSize());
 	return Reader->Close();
 }
 
+static bool ReadWeightmapFile(TArray<uint8>& Result, const FString& Filename, uint32 Flags = 0)
+{
+	bool bResult = ReadRawFile(Result, Filename, Flags);
+	if (bResult && Filename.EndsWith(TEXT(".png")))
+	{
+		IImageWrapperModule& ImageWrapperModule = FModuleManager::LoadModuleChecked<IImageWrapperModule>("ImageWrapper");
+		IImageWrapperPtr ImageWrapper = ImageWrapperModule.CreateImageWrapper(EImageFormat::PNG);
+		bResult = ImageWrapper->SetCompressed(Result.GetData(), Result.Num());
+		if (bResult)
+		{
+			const TArray<uint8>* RawData = nullptr;
+			bResult = ImageWrapper->GetRaw(ERGBFormat::Gray, 8, RawData);
+			if (bResult)
+			{
+				Result = *RawData;
+			}
+		}
+	}
+
+	return bResult;
+}
+
 static ULandscapeLayerInfoObject* GetandscapeLayerInfoObject(FName LayerName, FString ContentPath)
 {
 	// Build default layer object name and package name
@@ -1491,8 +1513,8 @@ void FWorldTileCollectionModel::ImportTiledLandscape_Executed()
 			TileImportSettings.ComponentSizeQuads	= ImportSettings.QuadsPerSection*ImportSettings.SectionsPerComponent;
 			TileImportSettings.QuadsPerSection		= ImportSettings.QuadsPerSection;
 			TileImportSettings.SectionsPerComponent = ImportSettings.SectionsPerComponent;
-			TileImportSettings.SizeX				= ImportSettings.TileResolution;
-			TileImportSettings.SizeY				= ImportSettings.TileResolution;
+			TileImportSettings.SizeX				= ImportSettings.SizeX;
+			TileImportSettings.SizeY				= ImportSettings.SizeX;
 			TileImportSettings.HeightmapFilename	= Filename;
 			TileImportSettings.LandscapeTransform.SetScale3D(TileScale);
 
@@ -1510,13 +1532,13 @@ void FWorldTileCollectionModel::ImportTiledLandscape_Executed()
 				if (WeightmapFile)
 				{
 					LayerImportInfo.SourceFilePath = *WeightmapFile;
-					ReadRawFile(LayerImportInfo.LayerData, *LayerImportInfo.SourceFilePath, FILEREAD_Silent);
+					ReadWeightmapFile(LayerImportInfo.LayerData, LayerImportInfo.SourceFilePath, FILEREAD_Silent);
 					LayerImportInfo.LayerInfo = GetandscapeLayerInfoObject(LayerImportInfo.LayerName, GetWorld()->GetOutermost()->GetName());
 					LayerImportInfo.LayerInfo->bNoWeightBlend = ImportSettings.LandscapeLayerSettingsList[LayerIdx].bNoBlendWeight;
 				}
 			}
 						
-			if (ReadRawFile(TileImportSettings.HeightData, *Filename, FILEREAD_Silent))
+			if (ReadRawFile(TileImportSettings.HeightData, Filename, FILEREAD_Silent))
 			{
 				FString MapFileName = WorldRootPath + TileName + FPackageName::GetMapPackageExtension();
 				// Create a new world - so we can 'borrow' its level
diff --git a/Engine/Source/Editor/WorldBrowser/WorldBrowser.Build.cs b/Engine/Source/Editor/WorldBrowser/WorldBrowser.Build.cs
--- a/Engine/Source/Editor/WorldBrowser/WorldBrowser.Build.cs
+++ b/Engine/Source/Editor/WorldBrowser/WorldBrowser.Build.cs
@@ -41,6 +41,7 @@ public class WorldBrowser : ModuleRules
                 "MeshUtilities",
                 "LandscapeEditor",
 				"Levels",
+                "ImageWrapper",
             }
 		);
 
@@ -50,6 +51,7 @@ public class WorldBrowser : ModuleRules
 				"SceneOutliner",
                 "MeshUtilities",
                 "ContentBrowser",
+                "ImageWrapper",
 			}
 		);
     }
diff --git a/Engine/Source/Runtime/Landscape/Classes/LandscapeHeightfieldCollisionComponent.h b/Engine/Source/Runtime/Landscape/Classes/LandscapeHeightfieldCollisionComponent.h
--- a/Engine/Source/Runtime/Landscape/Classes/LandscapeHeightfieldCollisionComponent.h
+++ b/Engine/Source/Runtime/Landscape/Classes/LandscapeHeightfieldCollisionComponent.h
@@ -180,7 +180,7 @@ class ULandscapeHeightfieldCollisionComponent : public UPrimitiveComponent
 	/** 
 	 * Cooks raw height data into collision object binary stream
 	 */
-	virtual bool CookCollsionData(const FName& Format, bool bUseOnlyDefMaterial, TArray<uint8>& OutCookedData, TArray<UPhysicalMaterial*>& OutMaterials) const;
+	virtual bool CookCollsionData(const FName& Format, bool bUseOnlyDefMaterial, bool bCheckDDC, TArray<uint8>& OutCookedData, TArray<UPhysicalMaterial*>& OutMaterials) const;
 
 	/** Modify a sub-region of the PhysX heightfield. Note that this does not update the physical material */
 	void UpdateHeightfieldRegion(int32 ComponentX1, int32 ComponentY1, int32 ComponentX2, int32 ComponentY2);
diff --git a/Engine/Source/Runtime/Landscape/Classes/LandscapeMeshCollisionComponent.h b/Engine/Source/Runtime/Landscape/Classes/LandscapeMeshCollisionComponent.h
--- a/Engine/Source/Runtime/Landscape/Classes/LandscapeMeshCollisionComponent.h
+++ b/Engine/Source/Runtime/Landscape/Classes/LandscapeMeshCollisionComponent.h
@@ -88,7 +88,7 @@ class ULandscapeMeshCollisionComponent : public ULandscapeHeightfieldCollisionCo
 	virtual void ExportCustomProperties(FOutputDevice& Out, uint32 Indent) override;
 	virtual void ImportCustomProperties(const TCHAR* SourceText, FFeedbackContext* Warn) override;
 
-	virtual bool CookCollsionData(const FName& Format, bool bUseOnlyDefMaterial, TArray<uint8>& OutCookedData, TArray<UPhysicalMaterial*>& OutMaterials) const override;
+	virtual bool CookCollsionData(const FName& Format, bool bUseOnlyDefMaterial, bool bCheckDDC, TArray<uint8>& OutCookedData, TArray<UPhysicalMaterial*>& OutMaterials) const override;
 #endif
 	// End UObject Interface.
 
diff --git a/Engine/Source/Runtime/Landscape/Private/Landscape.cpp b/Engine/Source/Runtime/Landscape/Private/Landscape.cpp
--- a/Engine/Source/Runtime/Landscape/Private/Landscape.cpp
+++ b/Engine/Source/Runtime/Landscape/Private/Landscape.cpp
@@ -1420,7 +1420,10 @@ void ALandscapeProxy::RecreateComponentsState()
 		if (Comp)
 		{
 			Comp->UpdateComponentToWorld();
-			Comp->RecreateCollision();
+#if WITH_EDITOR
+			Comp->UpdateAddCollisions();
+#endif
+			Comp->RecreatePhysicsState();
 		}
 	}
 }
diff --git a/Engine/Source/Runtime/Landscape/Private/LandscapeCollision.cpp b/Engine/Source/Runtime/Landscape/Private/LandscapeCollision.cpp
--- a/Engine/Source/Runtime/Landscape/Private/LandscapeCollision.cpp
+++ b/Engine/Source/Runtime/Landscape/Private/LandscapeCollision.cpp
@@ -260,13 +260,20 @@ void ULandscapeHeightfieldCollisionComponent::CreateCollisionObject()
 	// If we have not created a heightfield yet - do it now.
 	if (!IsValidRef(HeightfieldRef))
 	{
+		FPhysXHeightfieldRef* ExistingHeightfieldRef = nullptr;
+		bool bCheckDDC = true;
+
 		if (!HeightfieldGuid.IsValid())
 		{
 			HeightfieldGuid = FGuid::NewGuid();
+			bCheckDDC = false;
+		}
+		else
+		{
+			// Look for a heightfield object with the current Guid (this occurs with PIE)
+			ExistingHeightfieldRef = GSharedHeightfieldRefs.FindRef(HeightfieldGuid);
 		}
 
-		// Look for a heightfield object with the current Guid (this occurs with PIE)
-		FPhysXHeightfieldRef* ExistingHeightfieldRef = GSharedHeightfieldRefs.FindRef(HeightfieldGuid);
 		if (ExistingHeightfieldRef)
 		{
 			HeightfieldRef = ExistingHeightfieldRef;
@@ -276,7 +283,7 @@ void ULandscapeHeightfieldCollisionComponent::CreateCollisionObject()
 #if WITH_EDITOR
 			// Prepare heightfield data
 			static FName PhysicsFormatName(FPlatformProperties::GetPhysicsFormat());
-			CookCollsionData(PhysicsFormatName, false, CookedCollisionData, CookedPhysicalMaterials);
+			CookCollsionData(PhysicsFormatName, false, bCheckDDC, CookedCollisionData, CookedPhysicalMaterials);
 #endif //WITH_EDITOR
 
 			if (CookedCollisionData.Num())
@@ -306,7 +313,7 @@ void ULandscapeHeightfieldCollisionComponent::CreateCollisionObject()
 				if (!GetWorld()->IsGameWorld())
 				{
 					TArray<UPhysicalMaterial*> CookedMaterialsEd;
-					if (CookCollsionData(PhysicsFormatName, true, CookedCollisionDataEd, CookedMaterialsEd))
+					if (CookCollsionData(PhysicsFormatName, true, bCheckDDC, CookedCollisionDataEd, CookedMaterialsEd))
 					{
 						FPhysXInputStream HeightFieldStream(CookedCollisionDataEd.GetData(), CookedCollisionDataEd.Num());
 						HeightfieldRef->RBHeightfieldEd = GPhysXSDK->createHeightField(HeightFieldStream);
@@ -320,17 +327,20 @@ void ULandscapeHeightfieldCollisionComponent::CreateCollisionObject()
 }
 
 #if WITH_EDITOR
-bool ULandscapeHeightfieldCollisionComponent::CookCollsionData(const FName& Format, bool bUseDefMaterial, TArray<uint8>& OutCookedData, TArray<UPhysicalMaterial*>& OutMaterials) const
+bool ULandscapeHeightfieldCollisionComponent::CookCollsionData(const FName& Format, bool bUseDefMaterial, bool bCheckDDC, TArray<uint8>& OutCookedData, TArray<UPhysicalMaterial*>& OutMaterials) const
 {
 #if WITH_PHYSX
 
-	// Ensure that content was saved with physical materials before using DDC data
-	if (GetLinkerUE4Version() >= VER_UE4_LANDSCAPE_SERIALIZE_PHYSICS_MATERIALS) 
+	if (bCheckDDC)
 	{
-		if (GetDerivedDataCacheRef().GetSynchronous(*GetHFDDCKeyString(Format, bUseDefMaterial, HeightfieldGuid), OutCookedData))
+		// Ensure that content was saved with physical materials before using DDC data
+		if (GetLinkerUE4Version() >= VER_UE4_LANDSCAPE_SERIALIZE_PHYSICS_MATERIALS)
 		{
-			bShouldSaveCookedDataToDDC = false;
-			return true;
+			if (GetDerivedDataCacheRef().GetSynchronous(*GetHFDDCKeyString(Format, bUseDefMaterial, HeightfieldGuid), OutCookedData))
+			{
+				bShouldSaveCookedDataToDDC = false;
+				return true;
+			}
 		}
 	}
 				
@@ -452,17 +462,20 @@ bool ULandscapeHeightfieldCollisionComponent::CookCollsionData(const FName& Form
 	return false;
 }
 
-bool ULandscapeMeshCollisionComponent::CookCollsionData(const FName& Format, bool bUseDefMaterial, TArray<uint8>& OutCookedData, TArray<UPhysicalMaterial*>& OutMaterials) const
+bool ULandscapeMeshCollisionComponent::CookCollsionData(const FName& Format, bool bUseDefMaterial, bool bCheckDDC, TArray<uint8>& OutCookedData, TArray<UPhysicalMaterial*>& OutMaterials) const
 {
 #if WITH_PHYSX
 	
-	// Ensure that content was saved with physical materials before using DDC data
-	if (GetLinkerUE4Version() >= VER_UE4_LANDSCAPE_SERIALIZE_PHYSICS_MATERIALS)
+	if (bCheckDDC)
 	{
-		if (GetDerivedDataCacheRef().GetSynchronous(*GetHFDDCKeyString(Format, bUseDefMaterial, HeightfieldGuid), OutCookedData))
+		// Ensure that content was saved with physical materials before using DDC data
+		if (GetLinkerUE4Version() >= VER_UE4_LANDSCAPE_SERIALIZE_PHYSICS_MATERIALS)
 		{
-			bShouldSaveCookedDataToDDC = false;
-			return true;
+			if (GetDerivedDataCacheRef().GetSynchronous(*GetHFDDCKeyString(Format, bUseDefMaterial, HeightfieldGuid), OutCookedData))
+			{
+				bShouldSaveCookedDataToDDC = false;
+				return true;
+			}
 		}
 	}
 	
@@ -625,13 +638,20 @@ void ULandscapeMeshCollisionComponent::CreateCollisionObject()
 	// If we have not created a heightfield yet - do it now.
 	if (!IsValidRef(MeshRef))
 	{
+		FPhysXMeshRef* ExistingMeshRef = nullptr;
+		bool bCheckDDC = true;
+
 		if (!MeshGuid.IsValid())
 		{
 			MeshGuid = FGuid::NewGuid();
+			bCheckDDC = false;
+		}
+		else
+		{
+			// Look for a heightfield object with the current Guid (this occurs with PIE)
+			ExistingMeshRef = GSharedMeshRefs.FindRef(MeshGuid);
 		}
 
-		// Look for a heightfield object with the current Guid (this occurs with PIE)
-		FPhysXMeshRef* ExistingMeshRef = GSharedMeshRefs.FindRef(MeshGuid);
 		if (ExistingMeshRef)
 		{
 			MeshRef = ExistingMeshRef;
@@ -641,7 +661,7 @@ void ULandscapeMeshCollisionComponent::CreateCollisionObject()
 #if WITH_EDITOR
 			// Create cooked physics data
 			static FName PhysicsFormatName(FPlatformProperties::GetPhysicsFormat());
-			CookCollsionData(PhysicsFormatName, false, CookedCollisionData, CookedPhysicalMaterials);
+			CookCollsionData(PhysicsFormatName, false, bCheckDDC, CookedCollisionData, CookedPhysicalMaterials);
 #endif //WITH_EDITOR
 
 			if (CookedCollisionData.Num())
@@ -669,7 +689,7 @@ void ULandscapeMeshCollisionComponent::CreateCollisionObject()
 				if (!GetWorld()->IsGameWorld())
 				{
 					TArray<UPhysicalMaterial*> CookedMaterialsEd;
-					if (CookCollsionData(PhysicsFormatName, true, CookedCollisionDataEd, CookedMaterialsEd))
+					if (CookCollsionData(PhysicsFormatName, true, bCheckDDC, CookedCollisionDataEd, CookedMaterialsEd))
 					{
 						FPhysXInputStream MeshStream(CookedCollisionDataEd.GetData(), CookedCollisionDataEd.Num());
 						MeshRef->RBTriangleMeshEd = GPhysXSDK->createTriangleMesh(MeshStream);
@@ -1115,7 +1135,7 @@ void ULandscapeHeightfieldCollisionComponent::Serialize(FArchive& Ar)
 		if (Ar.IsCooking() && !HasAnyFlags(RF_ClassDefaultObject))
 		{
 			FName Format = Ar.CookingTarget()->GetPhysicsFormat(nullptr);
-			CookCollsionData(Format, false, CookedCollisionData, CookedPhysicalMaterials);
+			CookCollsionData(Format, false, true, CookedCollisionData, CookedPhysicalMaterials);
 			GetDerivedDataCacheRef().Put(*GetHFDDCKeyString(Format, false, HeightfieldGuid), CookedCollisionData);
 		}
 	}
