LOD maps - option to save generated assets in separate packages

diff --git a/Engine/Source/Developer/MeshUtilities/Private/MeshUtilities.cpp b/Engine/Source/Developer/MeshUtilities/Private/MeshUtilities.cpp
--- a/Engine/Source/Developer/MeshUtilities/Private/MeshUtilities.cpp
+++ b/Engine/Source/Developer/MeshUtilities/Private/MeshUtilities.cpp
@@ -3532,6 +3532,17 @@ void FMeshUtilities::CreateProxyMesh(
 		return;
 	}
 
+	// Base asset name for a new assets
+	// In case outer is null ProxyBasePackageName has to be long package name
+	if (InOuter == nullptr && FPackageName::IsShortPackageName(ProxyBasePackageName))
+	{
+		UE_LOG(LogMeshUtilities, Warning, TEXT("Invalid long package name: '%s'."), *ProxyBasePackageName);
+		return;
+	}
+	
+	const FString AssetBaseName = FPackageName::GetShortName(ProxyBasePackageName);
+	const FString AssetBasePath = InOuter ? TEXT("") : FPackageName::GetLongPackagePath(ProxyBasePackageName) + TEXT("/");
+	
 	TArray<ALandscapeProxy*>		LandscapesToMerge;
 	TArray<UStaticMeshComponent*>	ComponentsToMerge;
 	
@@ -3666,26 +3677,20 @@ void FMeshUtilities::CreateProxyMesh(
 		Vertex-= OutProxyLocation;
 	}
 	
-	//
-	// Base asset name for a new assets
-	//
-	const FString AssetBaseName = FPackageName::GetShortName(ProxyBasePackageName);
-	const FString AssetBasePath = FPackageName::IsShortPackageName(ProxyBasePackageName) ? 
-		FPackageName::FilenameToLongPackageName(FPaths::GameContentDir()) : (FPackageName::GetLongPackagePath(ProxyBasePackageName) + TEXT("/"));
-
 	// Construct proxy material
-	UMaterial* ProxyMaterial = MaterialExportUtils::CreateMaterial(ProxyFlattenMaterial, InOuter, ProxyBasePackageName, RF_Public|RF_Standalone);
+	UMaterial* ProxyMaterial = MaterialExportUtils::CreateMaterial(ProxyFlattenMaterial, InOuter, ProxyBasePackageName, RF_Public|RF_Standalone, OutAssetsToSync);
 	
 	// Construct proxy static mesh
 	UPackage* MeshPackage = InOuter;
+	FString MeshAssetName = TEXT("SM_") + AssetBaseName;
 	if (MeshPackage == nullptr)
 	{
-		MeshPackage = CreatePackage(NULL, *(AssetBasePath + TEXT("SM_") + AssetBaseName));
+		MeshPackage = CreatePackage(NULL, *(AssetBasePath + MeshAssetName));
 		MeshPackage->FullyLoad();
 		MeshPackage->Modify();
 	}
 
-	UStaticMesh* StaticMesh = new(MeshPackage, FName(*(TEXT("SM_") + AssetBaseName)), RF_Public|RF_Standalone) UStaticMesh(FObjectInitializer());
+	UStaticMesh* StaticMesh = new(MeshPackage, FName(*MeshAssetName), RF_Public|RF_Standalone) UStaticMesh(FObjectInitializer());
 	StaticMesh->InitResources();
 	{
 		FString OutputPath = StaticMesh->GetPathName();
@@ -3710,41 +3715,9 @@ void FMeshUtilities::CreateProxyMesh(
 
 		StaticMesh->Build();
 		StaticMesh->PostEditChange();
-	}
 
-	OutAssetsToSync.Add(ProxyMaterial);
-	OutAssetsToSync.Add(StaticMesh);
-
-#if 0 // dump flattened materials as texture assets
-	for (const auto& FlatMat : UniqueMaterials)
-	{
-		if (FlatMat.DiffuseSamples.Num() > 1)
-		{
-			FString DiffuseTextureName = MakeUniqueObjectName(Package, UTexture2D::StaticClass(), *(TEXT("T_FLATTEN_") + AssetBaseName + TEXT("_D"))).ToString();
-			
-			UPackage* TexPackage = CreatePackage(NULL, *(AssetBasePath + DiffuseTextureName));
-			TexPackage->FullyLoad();
-			TexPackage->Modify();
-					
-			FCreateTexture2DParameters TexParams;
-			TexParams.bUseAlpha = false;
-			TexParams.CompressionSettings = TC_Default;
-			TexParams.bDeferCompression = false;
-			TexParams.bSRGB = false;
-
-			UTexture2D* DiffuseTexture = FImageUtils::CreateTexture2D(
-				FlatMat.DiffuseSize.X, 
-				FlatMat.DiffuseSize.Y,
-				FlatMat.DiffuseSamples,
-				TexPackage,
-				DiffuseTextureName,
-				RF_Public|RF_Standalone, 
-				TexParams);
-
-			OutAssetsToSync.Add(DiffuseTexture);
-		}
+		OutAssetsToSync.Add(StaticMesh);
 	}
-#endif
 }
 
 bool FMeshUtilities::ConstructRawMesh(
diff --git a/Engine/Source/Editor/UnrealEd/Private/Dialogs/SMeshProxyDialog.cpp b/Engine/Source/Editor/UnrealEd/Private/Dialogs/SMeshProxyDialog.cpp
--- a/Engine/Source/Editor/UnrealEd/Private/Dialogs/SMeshProxyDialog.cpp
+++ b/Engine/Source/Editor/UnrealEd/Private/Dialogs/SMeshProxyDialog.cpp
@@ -602,7 +602,7 @@ FReply SMeshProxyDialog::OnMergeClicked()
 	{
 		GWarn->BeginSlowTask( LOCTEXT("MeshProxy_CreatingProxy", "Creating Mesh Proxy"), true);
 		GEditor->BeginTransaction( LOCTEXT( "MeshProxy_Create", "Creating Mesh Proxy" ) );
-
+		
 		FVector ProxyLocation = FVector::ZeroVector;
 		MeshUtilities.CreateProxyMesh(Actors, ProxySettings, NULL, ProxyPackageName, AssetsToSync, ProxyLocation);
 
@@ -694,6 +694,11 @@ FText SMeshProxyDialog::GetProxyPackagName() const
 void SMeshProxyDialog::OnPackageNameTextCommited(const FText& InText, ETextCommit::Type InCommitType)
 {
 	ProxyPackageName = InText.ToString();
+
+	if (FPackageName::IsShortPackageName(ProxyPackageName))
+	{
+		ProxyPackageName = FPackageName::FilenameToLongPackageName(FPaths::GameContentDir()) + ProxyPackageName;
+	}
 }
 
 FReply SMeshProxyDialog::OnSelectPackageNameClicked()
diff --git a/Engine/Source/Editor/UnrealEd/Private/EditorExporters.cpp b/Engine/Source/Editor/UnrealEd/Private/EditorExporters.cpp
--- a/Engine/Source/Editor/UnrealEd/Private/EditorExporters.cpp
+++ b/Engine/Source/Editor/UnrealEd/Private/EditorExporters.cpp
@@ -2849,11 +2849,18 @@ namespace MaterialExportUtils
 		return Texture;
 	}
 
-	UMaterial* CreateMaterial(const FFlattenMaterial& InFlattenMaterial, UPackage* Outer, const FString& BaseName, EObjectFlags Flags)
+	UMaterial* CreateMaterial(const FFlattenMaterial& InFlattenMaterial, UPackage* Outer, const FString& BaseName, EObjectFlags Flags, TArray<UObject*>& OutGeneratedAssets)
 	{
+		// Base name for a new assets
+		// In case outer is null BaseName has to be long package name
+		if (Outer == nullptr && FPackageName::IsShortPackageName(BaseName))
+		{
+			UE_LOG(LogEditorExporters, Warning, TEXT("Invalid long package name: '%s'."), *BaseName);
+			return nullptr;
+		}
+
 		const FString AssetBaseName = FPackageName::GetShortName(BaseName);
-		const FString AssetBasePath = FPackageName::IsShortPackageName(BaseName) ? 
-			FPackageName::FilenameToLongPackageName(FPaths::GameContentDir()) : (FPackageName::GetLongPackagePath(BaseName) + TEXT("/"));
+		const FString AssetBasePath = Outer ? TEXT("") : FPackageName::GetLongPackagePath(BaseName) + TEXT("/");
 				
 		// Create material
 		const FString MaterialAssetName = TEXT("M_") + AssetBaseName;
@@ -2868,6 +2875,7 @@ namespace MaterialExportUtils
 		UMaterial* Material = ConstructObject<UMaterial>(UMaterial::StaticClass(), MaterialOuter, FName(*MaterialAssetName), Flags);
 		Material->TwoSided = false;
 		Material->SetShadingModel(MSM_DefaultLit);
+		OutGeneratedAssets.Add(Material);
 
 		int32 MaterialNodeY = -150;
 		int32 MaterialNodeStepY = 180;
@@ -2879,6 +2887,7 @@ namespace MaterialExportUtils
 			const FString AssetLongName = AssetBasePath + AssetName;
 			const bool bSRGB = true;
 			UTexture2D* Texture = CreateTexture(Outer, AssetLongName, InFlattenMaterial.DiffuseSize, InFlattenMaterial.DiffuseSamples, TC_Default, TEXTUREGROUP_World, Flags, bSRGB);
+			OutGeneratedAssets.Add(Texture);
 			
 			UMaterialExpressionTextureSample* BasecolorExpression = ConstructObject<UMaterialExpressionTextureSample>(UMaterialExpressionTextureSample::StaticClass(), Material);
 			BasecolorExpression->Texture = Texture;
@@ -2892,12 +2901,12 @@ namespace MaterialExportUtils
 		}
 
 		// Metallic
-		if (InFlattenMaterial.RoughnessSamples.Num() > 1)
+		if (InFlattenMaterial.MetallicSamples.Num() > 1)
 		{
 			const FString AssetName = TEXT("T_") + AssetBaseName + TEXT("_M");
-			const FString AssetLongName = AssetBasePath + AssetName;
 			const bool bSRGB = false;
-			UTexture2D* Texture = CreateTexture(Outer, AssetBasePath, InFlattenMaterial.MetallicSize, InFlattenMaterial.MetallicSamples, TC_Grayscale, TEXTUREGROUP_World, Flags, bSRGB);
+			UTexture2D* Texture = CreateTexture(Outer, AssetBasePath + AssetName, InFlattenMaterial.MetallicSize, InFlattenMaterial.MetallicSamples, TC_Grayscale, TEXTUREGROUP_World, Flags, bSRGB);
+			OutGeneratedAssets.Add(Texture);
 			
 			UMaterialExpressionTextureSample* MetallicExpression = ConstructObject<UMaterialExpressionTextureSample>(UMaterialExpressionTextureSample::StaticClass(), Material);
 			MetallicExpression->Texture = Texture;
@@ -2927,9 +2936,9 @@ namespace MaterialExportUtils
 		if (InFlattenMaterial.SpecularSamples.Num() > 1)
 		{
 			const FString AssetName = TEXT("T_") + AssetBaseName + TEXT("_S");
-			const FString AssetLongName = AssetBasePath + AssetName;
 			const bool bSRGB = false;
-			UTexture2D* Texture = CreateTexture(Outer, AssetLongName, InFlattenMaterial.SpecularSize, InFlattenMaterial.SpecularSamples, TC_Default, TEXTUREGROUP_World, Flags, bSRGB);
+			UTexture2D* Texture = CreateTexture(Outer, AssetBasePath + AssetName, InFlattenMaterial.SpecularSize, InFlattenMaterial.SpecularSamples, TC_Default, TEXTUREGROUP_World, Flags, bSRGB);
+			OutGeneratedAssets.Add(Texture);
 			
 			UMaterialExpressionTextureSample* SpecularExpression = ConstructObject<UMaterialExpressionTextureSample>(UMaterialExpressionTextureSample::StaticClass(), Material);
 			SpecularExpression->Texture = Texture;
@@ -2946,9 +2955,9 @@ namespace MaterialExportUtils
 		if (InFlattenMaterial.RoughnessSamples.Num() > 1)
 		{
 			const FString AssetName = TEXT("T_") + AssetBaseName + TEXT("_R");
-			const FString AssetLongName = AssetBasePath + AssetName;
 			const bool bSRGB = false;
-			UTexture2D* Texture = CreateTexture(Outer, AssetLongName, InFlattenMaterial.RoughnessSize, InFlattenMaterial.RoughnessSamples, TC_Grayscale, TEXTUREGROUP_World, Flags, bSRGB);
+			UTexture2D* Texture = CreateTexture(Outer, AssetBasePath + AssetName, InFlattenMaterial.RoughnessSize, InFlattenMaterial.RoughnessSamples, TC_Grayscale, TEXTUREGROUP_World, Flags, bSRGB);
+			OutGeneratedAssets.Add(Texture);
 			
 			UMaterialExpressionTextureSample* RoughnessExpression = ConstructObject<UMaterialExpressionTextureSample>(UMaterialExpressionTextureSample::StaticClass(), Material);
 			RoughnessExpression->Texture = Texture;
@@ -2978,9 +2987,9 @@ namespace MaterialExportUtils
 		if (InFlattenMaterial.NormalSamples.Num() > 1)
 		{
 			const FString AssetName = TEXT("T_") + AssetBaseName + TEXT("_N");
-			const FString AssetLongName = AssetBasePath + AssetName;
 			const bool bSRGB = false;
-			UTexture2D* Texture = CreateTexture(Outer, AssetLongName, InFlattenMaterial.NormalSize, InFlattenMaterial.NormalSamples, TC_Normalmap, TEXTUREGROUP_WorldNormalMap, Flags, bSRGB);
+			UTexture2D* Texture = CreateTexture(Outer, AssetBasePath + AssetName, InFlattenMaterial.NormalSize, InFlattenMaterial.NormalSamples, TC_Normalmap, TEXTUREGROUP_WorldNormalMap, Flags, bSRGB);
+			OutGeneratedAssets.Add(Texture);
 			
 			UMaterialExpressionTextureSample* NormalExpression = ConstructObject<UMaterialExpressionTextureSample>(UMaterialExpressionTextureSample::StaticClass(), Material);
 			NormalExpression->Texture = Texture;
diff --git a/Engine/Source/Editor/UnrealEd/Public/MaterialExportUtils.h b/Engine/Source/Editor/UnrealEd/Public/MaterialExportUtils.h
--- a/Engine/Source/Editor/UnrealEd/Public/MaterialExportUtils.h
+++ b/Engine/Source/Editor/UnrealEd/Public/MaterialExportUtils.h
@@ -66,11 +66,12 @@ namespace MaterialExportUtils
 	/**
 	 * Creates UMaterial object from a flatten material
 	 *
-	 * @param Outer			Outer for the material and texture objects, if NULL new packages will be created for each asset
-	 * @param BaseName		BaseName for the material and texture objects, should be a long package name in case Outer is not specified
-	 * @param Flags			Object flags for the material and texture objects.
-	 * @return				Returns a pointer to the constructed UMaterial object.
+	 * @param Outer					Outer for the material and texture objects, if NULL new packages will be created for each asset
+	 * @param BaseName				BaseName for the material and texture objects, should be a long package name in case Outer is not specified
+	 * @param Flags					Object flags for the material and texture objects.
+	 * @param OutGeneratedAssets	List of generated assets - material, textures
+	 * @return						Returns a pointer to the constructed UMaterial object.
 	 */
-	UNREALED_API UMaterial* CreateMaterial(const FFlattenMaterial& InFlattenMaterial, UPackage* InOuter,const FString& BaseName, EObjectFlags Flags);
+	UNREALED_API UMaterial* CreateMaterial(const FFlattenMaterial& InFlattenMaterial, UPackage* InOuter, const FString& BaseName, EObjectFlags Flags, TArray<UObject*>& OutGeneratedAssets);
 }
 
diff --git a/Engine/Source/Editor/WorldBrowser/Private/Tiles/WorldTileCollectionModel.cpp b/Engine/Source/Editor/WorldBrowser/Private/Tiles/WorldTileCollectionModel.cpp
--- a/Engine/Source/Editor/WorldBrowser/Private/Tiles/WorldTileCollectionModel.cpp
+++ b/Engine/Source/Editor/WorldBrowser/Private/Tiles/WorldTileCollectionModel.cpp
@@ -1878,7 +1878,7 @@ bool FWorldTileCollectionModel::GenerateLODLevels(FLevelModelList InLevelList, i
 		LODPackage->Modify();
 		// This is a hack to avoid save file dialog when we will be saving LOD map package
 		LODPackage->FileName = FName(*LODLevelFileName);
-				
+
 		// This is current actors offset from their original position
 		FVector ActorsOffset = FVector(TileModel->GetAbsoluteLevelPosition() - GetWorldOriginLocationXY(GetWorld()));
 		if (GetWorld()->WorldComposition->bTemporallyDisableOriginTracking)
@@ -1890,6 +1890,13 @@ bool FWorldTileCollectionModel::GenerateLODLevels(FLevelModelList InLevelList, i
 		TArray<FTransform>			AssetsToSpawnTransform;
 		TArray<AActor*>				Actors;
 		TArray<ALandscapeProxy*>	LandscapeActors;
+		TArray<UObject*>			GeneratedAssets;
+		
+		// Where generated assets will be stored
+		UPackage* AssetsOuter = SimplificationDetails.bCreatePackagePerAsset ? nullptr : LODPackage;
+		// In case we don't have outer generated assets should have same path as LOD level
+		const FString AssetsPath = AssetsOuter ? TEXT("") : FPackageName::GetLongPackagePath(LODLevelPackageName) + TEXT("/");
+		
 		// Separate landscape actors from all others
 		for (AActor* Actor : TileModel->GetLevelObject()->Actors)
 		{
@@ -1918,7 +1925,7 @@ bool FWorldTileCollectionModel::GenerateLODLevels(FLevelModelList InLevelList, i
 			FVector OutProxyLocation;
 			FString ProxyPackageName = TEXT("PROXY_") + FPackageName::GetShortName(TileModel->TileDetails->PackageName);
 			// Generate proxy mesh and proxy material assets
-			MeshUtilities.CreateProxyMesh(Actors, ProxySettings, LODPackage, ProxyPackageName, OutAssets, OutProxyLocation);
+			MeshUtilities.CreateProxyMesh(Actors, ProxySettings, AssetsOuter, AssetsPath + ProxyPackageName, OutAssets, OutProxyLocation);
 		
 			if (OutAssets.Num())
 			{
@@ -1928,6 +1935,8 @@ bool FWorldTileCollectionModel::GenerateLODLevels(FLevelModelList InLevelList, i
 					AssetsToSpawn.Add(ProxyMesh);
 					AssetsToSpawnTransform.Add(FTransform(OutProxyLocation - ActorsOffset));
 				}
+
+				GeneratedAssets.Append(OutAssets);
 			}
 		}
 
@@ -1967,7 +1976,7 @@ bool FWorldTileCollectionModel::GenerateLODLevels(FLevelModelList InLevelList, i
 				}
 			}
 								
-			// This is texture resolution for a landscape, probably needs to be calculated using landscape size
+			// This is texture resolution for a landscape mesh, probably needs to be calculated using landscape size
 			const FIntPoint LandscapeTextureSize(1024, 1024);
 			LandscapeFlattenMaterial.DiffuseSize	= LandscapeTextureSize;
 			LandscapeFlattenMaterial.NormalSize		= SimplificationDetails.bGenerateLandscapeNormalMap ? LandscapeTextureSize : FIntPoint::ZeroValue;
@@ -1980,14 +1989,22 @@ bool FWorldTileCollectionModel::GenerateLODLevels(FLevelModelList InLevelList, i
 			FString LandscapeBaseAssetName = Landscape->GetName();
 			// Construct landscape material
 			UMaterial* StaticLandscapeMaterial = MaterialExportUtils::CreateMaterial(
-				LandscapeFlattenMaterial, LODPackage, *LandscapeBaseAssetName, RF_Public|RF_Standalone);
+				LandscapeFlattenMaterial, AssetsOuter, *(AssetsPath + LandscapeBaseAssetName), RF_Public|RF_Standalone, GeneratedAssets);
 			// Currently landscape exports world space normal map
 			StaticLandscapeMaterial->bTangentSpaceNormal = false;
 			StaticLandscapeMaterial->PostEditChange();
 	
 			// Construct landscape static mesh
 			FString LandscapeMeshAssetName = TEXT("SM_") + LandscapeBaseAssetName;
-			UStaticMesh* StaticMesh = new(LODPackage, *LandscapeMeshAssetName, RF_Public|RF_Standalone) UStaticMesh(FObjectInitializer());
+			UPackage* MeshOuter = AssetsOuter;
+			if (MeshOuter == nullptr)
+			{
+				MeshOuter = CreatePackage(nullptr, *(AssetsPath + LandscapeMeshAssetName));
+				MeshOuter->FullyLoad();
+				MeshOuter->Modify();
+			}
+					
+			UStaticMesh* StaticMesh = new(MeshOuter, *LandscapeMeshAssetName, RF_Public|RF_Standalone) UStaticMesh(FObjectInitializer());
 			StaticMesh->InitResources();
 			{
 				FString OutputPath = StaticMesh->GetPathName();
@@ -2014,6 +2031,7 @@ bool FWorldTileCollectionModel::GenerateLODLevels(FLevelModelList InLevelList, i
 				StaticMesh->PostEditChange();
 			}
 
+			GeneratedAssets.Add(StaticMesh);
 			AssetsToSpawn.Add(StaticMesh);
 			AssetsToSpawnTransform.Add(FTransform(LandscapeWorldLocation - ActorsOffset));
 
@@ -2027,9 +2045,30 @@ bool FWorldTileCollectionModel::GenerateLODLevels(FLevelModelList InLevelList, i
 			GetWorld()->WorldComposition->bTemporallyDisableOriginTracking = false;
 		}
 	
-		// Create new level and spawn generated assets in it
 		if (AssetsToSpawn.Num())
 		{
+			// Save generated assets
+			if (SimplificationDetails.bCreatePackagePerAsset && GeneratedAssets.Num())
+			{							
+				const bool bCheckDirty = false;
+				const bool bPromptToSave = false;
+				TArray<UPackage*> PackagesToSave;
+
+				for (UObject* Asset : GeneratedAssets)
+				{
+					FAssetRegistryModule::AssetCreated(Asset);
+					Editor->BroadcastObjectReimported(Asset);
+					PackagesToSave.Add(Asset->GetOutermost());
+				}
+								
+				FEditorFileUtils::PromptForCheckoutAndSave(PackagesToSave, bCheckDirty, bPromptToSave);
+
+				// Also notify the content browser that the new assets exists
+				//FContentBrowserModule& ContentBrowserModule = FModuleManager::Get().LoadModuleChecked<FContentBrowserModule>("ContentBrowser");
+				//ContentBrowserModule.Get().SyncBrowserToAssets(GeneratedAssets, true);
+			}
+			
+			// Create new level and spawn generated assets in it
 			UWorld* LODWorld = UWorld::FindWorldInPackage(LODPackage);
 			if (LODWorld)
 			{
@@ -2054,6 +2093,7 @@ bool FWorldTileCollectionModel::GenerateLODLevels(FLevelModelList InLevelList, i
 			if (FEditorFileUtils::PromptToCheckoutLevels(false, LODWorld->PersistentLevel))
 			{
 				FEditorFileUtils::SaveLevel(LODWorld->PersistentLevel, *LODLevelFileName);
+				FAssetRegistryModule::AssetCreated(LODWorld->PersistentLevel);
 			}
 			
 			// Destroy the new world we created and collect the garbage
diff --git a/Engine/Source/Runtime/Engine/Classes/Engine/Level.h b/Engine/Source/Runtime/Engine/Classes/Engine/Level.h
--- a/Engine/Source/Runtime/Engine/Classes/Engine/Level.h
+++ b/Engine/Source/Runtime/Engine/Classes/Engine/Level.h
@@ -237,6 +237,10 @@ struct ENGINE_API FLevelSimplificationDetails
 	// Percentage of details relative to main tile details
 	UPROPERTY(Category=ReductionSettings, EditAnywhere, meta=(ClampMin = "0", ClampMax = "100", UIMin = "0", UIMax = "100"))	
 	float DetailsPercentage;
+
+	// Whether to create separate packages for each generated asset. All in map package otherwise
+	UPROPERTY(Category=Landscape, EditAnywhere)
+	bool bCreatePackagePerAsset;
 	
 	// Landscape LOD to use for static mesh generation
 	UPROPERTY(Category=Landscape, EditAnywhere, meta=(ClampMin = "0", ClampMax = "7", UIMin = "0", UIMax = "7"))
diff --git a/Engine/Source/Runtime/Engine/Private/Level.cpp b/Engine/Source/Runtime/Engine/Private/Level.cpp
--- a/Engine/Source/Runtime/Engine/Private/Level.cpp
+++ b/Engine/Source/Runtime/Engine/Private/Level.cpp
@@ -177,6 +177,7 @@ bool FLevelSimplificationDetails::operator == (const FLevelSimplificationDetails
 {
 	return
 		DetailsPercentage == Other.DetailsPercentage &&
+		bCreatePackagePerAsset == Other.bCreatePackagePerAsset &&
 		LandscapeExportLOD == Other.LandscapeExportLOD &&
 		bGenerateLandscapeNormalMap == Other.bGenerateLandscapeNormalMap &&
 		bGenerateLandscapeMetallicMap == Other.bGenerateLandscapeMetallicMap &&
