World composition: Show both - observer and player markers in minimap

diff --git a/Engine/Source/Editor/WorldBrowser/Private/LevelCollectionModel.cpp b/Engine/Source/Editor/WorldBrowser/Private/LevelCollectionModel.cpp
--- a/Engine/Source/Editor/WorldBrowser/Private/LevelCollectionModel.cpp
+++ b/Engine/Source/Editor/WorldBrowser/Private/LevelCollectionModel.cpp
@@ -686,6 +686,11 @@ void FLevelCollectionModel::CustomizeFileMainMenu(FMenuBuilder& InMenuBuilder) c
 	}
 }
 
+bool FLevelCollectionModel::GetPlayerView(FVector& Location, FRotator& Rotation) const
+{
+	return false;
+}
+
 bool FLevelCollectionModel::GetObserverView(FVector& Location, FRotator& Rotation) const
 {
 	return false;
diff --git a/Engine/Source/Editor/WorldBrowser/Private/LevelCollectionModel.h b/Engine/Source/Editor/WorldBrowser/Private/LevelCollectionModel.h
--- a/Engine/Source/Editor/WorldBrowser/Private/LevelCollectionModel.h
+++ b/Engine/Source/Editor/WorldBrowser/Private/LevelCollectionModel.h
@@ -129,8 +129,11 @@ public:
 	
 	/**	Customize 'File' section in main menu  */
 	virtual void CustomizeFileMainMenu(FMenuBuilder& InMenuBuilder) const;
+		
+	/**	@return	Player view in the PIE/Simulation world */
+	virtual bool GetPlayerView(FVector& Location, FRotator& Rotation) const;
 
-	/**	@return	Observer view in the world, usually camera view */
+	/**	@return	Observer view in the Editor/Similuation world */
 	virtual bool GetObserverView(FVector& Location, FRotator& Rotation) const;
 
 	/**	Compares 2 levels by Z order */
diff --git a/Engine/Source/Editor/WorldBrowser/Private/Tiles/SWorldComposition.cpp b/Engine/Source/Editor/WorldBrowser/Private/Tiles/SWorldComposition.cpp
--- a/Engine/Source/Editor/WorldBrowser/Private/Tiles/SWorldComposition.cpp
+++ b/Engine/Source/Editor/WorldBrowser/Private/Tiles/SWorldComposition.cpp
@@ -375,6 +375,33 @@ public:
 					FSlateDrawElement::RelativeToElement
 					);
 			}
+
+			FVector PlayerPosition;
+			FRotator PlayerRotation;
+			if (WorldModel->GetPlayerView(PlayerPosition, PlayerRotation))
+			{
+				FVector2D PlayerPositionScreen = GraphCoordToPanelCoord(FVector2D(PlayerPosition.X, PlayerPosition.Y));
+				const FSlateBrush* CameraImage = FEditorStyle::GetBrush(TEXT("WorldBrowser.SimulationViewPositon"));
+	
+				FPaintGeometry PaintGeometry = AllottedGeometry.ToPaintGeometry(
+					PlayerPositionScreen - CameraImage->ImageSize*0.5f, 
+					CameraImage->ImageSize
+					);
+
+				FSlateDrawElement::MakeRotatedBox(
+					OutDrawElements,
+					++LayerId,
+					PaintGeometry,
+					CameraImage,
+					MyClippingRect,
+					ESlateDrawEffect::None,
+					FMath::DegreesToRadians(PlayerRotation.Yaw),
+					CameraImage->ImageSize*0.5f,
+					FSlateDrawElement::RelativeToElement,
+					FLinearColor(FColorList::Orange)
+					);
+			}
+
 		}
 
 		LayerId = PaintScaleRuler(AllottedGeometry, MyClippingRect, OutDrawElements, LayerId);
diff --git a/Engine/Source/Editor/WorldBrowser/Private/Tiles/WorldTileCollectionModel.cpp b/Engine/Source/Editor/WorldBrowser/Private/Tiles/WorldTileCollectionModel.cpp
--- a/Engine/Source/Editor/WorldBrowser/Private/Tiles/WorldTileCollectionModel.cpp
+++ b/Engine/Source/Editor/WorldBrowser/Private/Tiles/WorldTileCollectionModel.cpp
@@ -571,30 +571,36 @@ void FWorldTileCollectionModel::CustomizeFileMainMenu(FMenuBuilder& InMenuBuilde
 	InMenuBuilder.EndSection();
 }
 
-bool FWorldTileCollectionModel::GetObserverView(FVector& Location, FRotator& Rotation) const
+bool FWorldTileCollectionModel::GetPlayerView(FVector& Location, FRotator& Rotation) const
 {
-	const UEditorEngine* EditorEngine = Editor.Get();
-
 	if (IsSimulating())
 	{
-		if (EditorEngine->bIsSimulatingInEditor && GCurrentLevelEditingViewportClient->IsSimulateInEditorViewport())
+		UWorld* SimulationWorld = GetSimulationWorld();
+		for (FConstPlayerControllerIterator Iterator = SimulationWorld->GetPlayerControllerIterator(); Iterator; ++Iterator)
 		{
-			Rotation = GCurrentLevelEditingViewportClient->GetViewRotation();
-			Location = GCurrentLevelEditingViewportClient->GetViewLocation();
+			APlayerController* PlayerActor = *Iterator;
+			PlayerActor->GetPlayerViewPoint(Location, Rotation);
 			return true;
 		}
-		else
-		{
-			UWorld* SimulationWorld = GetSimulationWorld();
-			for (FConstPlayerControllerIterator Iterator = SimulationWorld->GetPlayerControllerIterator(); Iterator; ++Iterator)
-			{
-				APlayerController* PlayerActor = *Iterator;
-				PlayerActor->GetPlayerViewPoint(Location, Rotation);
-				return true;
-			}
-		}
 	}
-	else
+	
+	return false;
+}
+
+bool FWorldTileCollectionModel::GetObserverView(FVector& Location, FRotator& Rotation) const
+{
+	const UEditorEngine* EditorEngine = Editor.Get();
+
+	// We are in the SIE
+	if (EditorEngine->bIsSimulatingInEditor && GCurrentLevelEditingViewportClient->IsSimulateInEditorViewport())
+	{
+		Rotation = GCurrentLevelEditingViewportClient->GetViewRotation();
+		Location = GCurrentLevelEditingViewportClient->GetViewLocation();
+		return true;
+	}
+
+	// We are in the editor world
+	if (EditorEngine->PlayWorld == nullptr)
 	{
 		for (const FLevelEditorViewportClient* ViewportClient : EditorEngine->LevelViewportClients)
 		{
@@ -606,7 +612,7 @@ bool FWorldTileCollectionModel::GetObserverView(FVector& Location, FRotator& Rot
 			}
 		}
 	}
-	
+
 	return false;
 }
 
diff --git a/Engine/Source/Editor/WorldBrowser/Private/Tiles/WorldTileCollectionModel.h b/Engine/Source/Editor/WorldBrowser/Private/Tiles/WorldTileCollectionModel.h
--- a/Engine/Source/Editor/WorldBrowser/Private/Tiles/WorldTileCollectionModel.h
+++ b/Engine/Source/Editor/WorldBrowser/Private/Tiles/WorldTileCollectionModel.h
@@ -34,6 +34,7 @@ public:
 	virtual bool PassesAllFilters(TSharedPtr<FLevelModel> InLevelModel) const override;
 	virtual void BuildHierarchyMenu(FMenuBuilder& InMenuBuilder) const override;
 	virtual void CustomizeFileMainMenu(FMenuBuilder& InMenuBuilder) const override;
+	virtual bool GetPlayerView(FVector& Location, FRotator& Rotation) const override;
 	virtual bool GetObserverView(FVector& Location, FRotator& Rotation) const override;
 	virtual bool CompareLevelsZOrder(TSharedPtr<FLevelModel> InA, TSharedPtr<FLevelModel> InB) const override;
 	virtual void RegisterDetailsCustomization(class FPropertyEditorModule& PropertyModule, TSharedPtr<class IDetailsView> InDetailsView)  override;
