World composition: tiled landscape import tool creates one empty landscape actor in the persistent level and landscape proxy in own sub-level for each tile

diff --git a/Engine/Source/Editor/WorldBrowser/Private/Tiles/WorldTileCollectionModel.cpp b/Engine/Source/Editor/WorldBrowser/Private/Tiles/WorldTileCollectionModel.cpp
--- a/Engine/Source/Editor/WorldBrowser/Private/Tiles/WorldTileCollectionModel.cpp
+++ b/Engine/Source/Editor/WorldBrowser/Private/Tiles/WorldTileCollectionModel.cpp
@@ -20,6 +20,7 @@
 #include "Engine/StaticMeshActor.h"
 #include "Engine/StaticMesh.h"
 #include "Engine/LevelStreaming.h"
+#include "Landscape.h"
 #include "GameFramework/WorldSettings.h"
 #include "Engine/Light.h"
 #include "Foliage/InstancedFoliageActor.h"
@@ -1460,7 +1461,7 @@ static bool ReadWeightmapFile(TArray<uint8>& Result, const FString& Filename, ui
 	return bResult;
 }
 
-static ULandscapeLayerInfoObject* GetandscapeLayerInfoObject(FName LayerName, FString ContentPath)
+static ULandscapeLayerInfoObject* GetLandscapeLayerInfoObject(FName LayerName, const FString& ContentPath)
 {
 	// Build default layer object name and package name
 	FString LayerObjectName = FString::Printf(TEXT("%s_LayerInfo"), *LayerName.ToString());
@@ -1495,6 +1496,33 @@ static ULandscapeLayerInfoObject* GetandscapeLayerInfoObject(FName LayerName, FS
 	return LayerInfo;
 }
 
+static void SetupLandscapeImportLayers(const FTiledLandscapeImportSettings& InImportSettings, const FString& ContentPath, int32 TileIndex, TArray<FLandscapeImportLayerInfo>& OutLayerInfo)
+{
+	for (const auto& LayerSettings : InImportSettings.LandscapeLayerSettingsList)
+	{
+		FLandscapeImportLayerInfo LayerImportInfo(LayerSettings.Name);
+		
+		// Do we have a weightmap data for this tile?
+		const FString* WeightmapFile = nullptr;
+		if (InImportSettings.TileCoordinates.IsValidIndex(TileIndex))
+		{
+			FIntPoint TileCoordinates = InImportSettings.TileCoordinates[TileIndex];
+			WeightmapFile = LayerSettings.WeightmapFiles.Find(TileCoordinates);
+			if (WeightmapFile)
+			{
+				LayerImportInfo.SourceFilePath = *WeightmapFile;
+				ReadWeightmapFile(LayerImportInfo.LayerData, LayerImportInfo.SourceFilePath, FILEREAD_Silent);
+			}
+		}
+		
+		LayerImportInfo.LayerInfo = GetLandscapeLayerInfoObject(LayerImportInfo.LayerName, ContentPath);
+		LayerImportInfo.LayerInfo->bNoWeightBlend = LayerSettings.bNoBlendWeight;
+						
+		OutLayerInfo.Add(LayerImportInfo);
+	}
+}
+
+
 void FWorldTileCollectionModel::ImportTiledLandscape_Executed()
 {
 	/** Create the window to host widget */
@@ -1527,13 +1555,36 @@ void FWorldTileCollectionModel::ImportTiledLandscape_Executed()
 
 		GWarn->BeginSlowTask(LOCTEXT("ImportingLandscapeTilesBegin", "Importing landscape tiles"), true);
 		
-		// Shared GUID used between landscape proxies
-		FGuid LandscapeGuid;
+		// Create main landscape actor in persistent level, it will be empty (no components in it)
+		// All landscape tiles will go into it's own sub-levels
+		FGuid LandscapeGuid = FGuid::NewGuid();
+		{
+			ALandscape* Landscape = Cast<UWorld>(GetWorld())->SpawnActor<ALandscape>();
+			Landscape->SetActorTransform(FTransform(FQuat::Identity, FVector::ZeroVector, ImportSettings.Scale3D));
+
+			// Setup layers list for importing
+			TArray<FLandscapeImportLayerInfo> ImportLayers;
+			SetupLandscapeImportLayers(ImportSettings, GetWorld()->GetOutermost()->GetName(), INDEX_NONE, ImportLayers);
+			
+			// Set landscape configuration
+			Landscape->LandscapeMaterial	= ImportSettings.LandscapeMaterial.Get();
+			Landscape->ComponentSizeQuads	= ImportSettings.QuadsPerSection*ImportSettings.SectionsPerComponent;
+			Landscape->NumSubsections		= ImportSettings.SectionsPerComponent;
+			Landscape->SubsectionSizeQuads	= ImportSettings.QuadsPerSection;
+			Landscape->SetLandscapeGuid(LandscapeGuid);
+			for (const auto& ImportLayerInfo : ImportLayers)
+			{
+				Landscape->EditorLayerSettings.Add(FLandscapeEditorLayerSettings(ImportLayerInfo.LayerInfo));
+			}
+			ULandscapeInfo::RecreateLandscapeInfo(GetWorld(), false);
+		}
 
 		// Import tiles
 		int32 TileIndex = 0;
 		for (const FString& Filename : ImportSettings.HeightmapFileList)
 		{
+			check(LandscapeGuid.IsValid());
+									
 			FString TileName = FPaths::GetBaseFilename(Filename);
 			FVector TileScale = ImportSettings.Scale3D;
 
@@ -1550,25 +1601,8 @@ void FWorldTileCollectionModel::ImportTiledLandscape_Executed()
 			TileImportSettings.HeightmapFilename	= Filename;
 			TileImportSettings.LandscapeTransform.SetScale3D(TileScale);
 
-			// Setup weightmaps for each layer
-			for (int32 LayerIdx = 0; LayerIdx < ImportSettings.LandscapeLayerSettingsList.Num(); ++LayerIdx)
-			{
-				FName LayerName = ImportSettings.LandscapeLayerSettingsList[LayerIdx].Name;
-				
-				TileImportSettings.ImportLayers.Add(FLandscapeImportLayerInfo(LayerName));
-				FLandscapeImportLayerInfo& LayerImportInfo = TileImportSettings.ImportLayers.Last();
-				
-				// Do we have a weightmap for this tile?
-				FIntPoint TileCoordinates = ImportSettings.TileCoordinates[TileIndex];
-				const FString* WeightmapFile = ImportSettings.LandscapeLayerSettingsList[LayerIdx].WeightmapFiles.Find(TileCoordinates);
-				if (WeightmapFile)
-				{
-					LayerImportInfo.SourceFilePath = *WeightmapFile;
-					ReadWeightmapFile(LayerImportInfo.LayerData, LayerImportInfo.SourceFilePath, FILEREAD_Silent);
-					LayerImportInfo.LayerInfo = GetandscapeLayerInfoObject(LayerImportInfo.LayerName, GetWorld()->GetOutermost()->GetName());
-					LayerImportInfo.LayerInfo->bNoWeightBlend = ImportSettings.LandscapeLayerSettingsList[LayerIdx].bNoBlendWeight;
-				}
-			}
+			// Setup layers list for importing
+			SetupLandscapeImportLayers(ImportSettings, GetWorld()->GetOutermost()->GetName(), TileIndex, TileImportSettings.ImportLayers);
 						
 			if (ReadRawFile(TileImportSettings.HeightData, Filename, FILEREAD_Silent))
 			{
@@ -1586,34 +1620,11 @@ void FWorldTileCollectionModel::ImportTiledLandscape_Executed()
 					// Hide level, so we do not depend on a current world origin
 					NewTileModel->SetVisible(false);
 					
-					// Create landscape (actor/proxy) in a new level
-					ALandscapeProxy* NewLandscape = NewTileModel->ImportLandscape(TileImportSettings);
-
-					// Use first created landscape as a landscape actor and all next tiles as landscape proxies 
-					if (!LandscapeGuid.IsValid())
-					{
-						LandscapeGuid = NewLandscape->GetLandscapeGuid();
-					}
+					// Create landscape proxy in a new level
+					ALandscapeProxy* NewLandscape = NewTileModel->ImportLandscapeTile(TileImportSettings);
 
 					if (NewLandscape)
 					{
-						// update shared LandscapeInfo objects with new layers
-						ULandscapeInfo* LandscapeInfo = NewLandscape->GetLandscapeInfo(true);
-						LandscapeInfo->UpdateLayerInfoMap(NewLandscape);
-						for (const FLandscapeImportLayerInfo& LayerImportInfo : TileImportSettings.ImportLayers)
-						{
-							if (LayerImportInfo.LayerInfo)
-							{
-								NewLandscape->EditorLayerSettings.Add(FLandscapeEditorLayerSettings(LayerImportInfo.LayerInfo, LayerImportInfo.SourceFilePath));
-								int32 LayerInfoIndex = LandscapeInfo->GetLayerInfoIndex(LayerImportInfo.LayerInfo->LayerName);
-								if (ensure(LayerInfoIndex != INDEX_NONE))
-								{
-									FLandscapeInfoLayerSettings& LayerSettings = LandscapeInfo->Layers[LayerInfoIndex];
-									LayerSettings.LayerInfoObj = LayerImportInfo.LayerInfo;
-								}
-							}
-						}
-						
 						// Set landscape actor location at min corner of bounding rect, so it will be centered around origin
 						FIntRect NewLandscapeRect = NewLandscape->GetBoundingRect();
 						float WidthX = NewLandscapeRect.Width()*TileScale.X;
diff --git a/Engine/Source/Editor/WorldBrowser/Private/Tiles/WorldTileModel.cpp b/Engine/Source/Editor/WorldBrowser/Private/Tiles/WorldTileModel.cpp
--- a/Engine/Source/Editor/WorldBrowser/Private/Tiles/WorldTileModel.cpp
+++ b/Engine/Source/Editor/WorldBrowser/Private/Tiles/WorldTileModel.cpp
@@ -9,7 +9,6 @@
 #include "WorldTileModel.h"
 #include "Engine/LevelBounds.h"
 #include "Engine/WorldComposition.h"
-#include "Landscape.h"
 #include "GameFramework/WorldSettings.h"
 #include "Engine/LevelStreaming.h"
 #include "Engine/LevelStreamingKismet.h"
@@ -754,7 +753,7 @@ bool FWorldTileModel::CreateAdjacentLandscapeProxy(ALandscapeProxy* SourceLandsc
 	ImportSettings.LandscapeTransform.SetScale3D(SourceLandscapeScale);
 	
 	// Create new landscape object
-	ALandscapeProxy* AdjacenLandscape = ImportLandscape(ImportSettings);
+	ALandscapeProxy* AdjacenLandscape = ImportLandscapeTile(ImportSettings);
 	if (AdjacenLandscape)
 	{
 		// Copy source landscape properties 
@@ -795,28 +794,18 @@ bool FWorldTileModel::CreateAdjacentLandscapeProxy(ALandscapeProxy* SourceLandsc
 	return false;
 }
 
-ALandscapeProxy* FWorldTileModel::ImportLandscape(const FLandscapeImportSettings& Settings)
+ALandscapeProxy* FWorldTileModel::ImportLandscapeTile(const FLandscapeImportSettings& Settings)
 {
 	if (!IsLoaded())
 	{
 		return nullptr;
 	}
 	
-	ALandscapeProxy*	LandscapeProxy;
-	FGuid				LandscapeGuid = Settings.LandscapeGuid;
-	
-	if (LandscapeGuid.IsValid())
-	{
-		LandscapeProxy = Cast<UWorld>(LoadedLevel->GetOuter())->SpawnActor<ALandscapeProxy>();
-	}
-	else
-	{
-		LandscapeProxy = Cast<UWorld>(LoadedLevel->GetOuter())->SpawnActor<ALandscape>();
-		LandscapeGuid = FGuid::NewGuid();
-	}
+	check(Settings.LandscapeGuid.IsValid())
 	
+	ALandscapeProxy* LandscapeProxy = Cast<UWorld>(LoadedLevel->GetOuter())->SpawnActor<ALandscapeProxy>();
 	LandscapeProxy->SetActorTransform(Settings.LandscapeTransform);
-	
+		
 	if (Settings.LandscapeMaterial)
 	{
 		LandscapeProxy->LandscapeMaterial = Settings.LandscapeMaterial;
@@ -824,7 +813,7 @@ ALandscapeProxy* FWorldTileModel::ImportLandscape(const FLandscapeImportSettings
 	
 	// Create landscape components	
 	LandscapeProxy->Import(
-		LandscapeGuid, 
+		Settings.LandscapeGuid, 
 		Settings.SizeX, 
 		Settings.SizeY, 
 		Settings.ComponentSizeQuads, 
diff --git a/Engine/Source/Editor/WorldBrowser/Private/Tiles/WorldTileModel.h b/Engine/Source/Editor/WorldBrowser/Private/Tiles/WorldTileModel.h
--- a/Engine/Source/Editor/WorldBrowser/Private/Tiles/WorldTileModel.h
+++ b/Engine/Source/Editor/WorldBrowser/Private/Tiles/WorldTileModel.h
@@ -154,7 +154,7 @@ public:
 	bool CreateAdjacentLandscapeProxy(ALandscapeProxy* SourceLandscape, FIntPoint SourceTileOffset, FWorldTileModel::EWorldDirections InWhere);
 
 	/**  */
-	ALandscapeProxy* ImportLandscape(const FLandscapeImportSettings& Settings);
+	ALandscapeProxy* ImportLandscapeTile(const FLandscapeImportSettings& Settings);
 
 private:
 	/** Flush world info to package and level objects */
diff --git a/Engine/Source/Runtime/Landscape/Classes/LandscapeProxy.h b/Engine/Source/Runtime/Landscape/Classes/LandscapeProxy.h
--- a/Engine/Source/Runtime/Landscape/Classes/LandscapeProxy.h
+++ b/Engine/Source/Runtime/Landscape/Classes/LandscapeProxy.h
@@ -426,6 +426,7 @@ public:
 #endif	//WITH_EDITOR
 
 	FGuid GetLandscapeGuid() const { return LandscapeGuid; }
+	void SetLandscapeGuid(const FGuid& Guid) { LandscapeGuid = Guid; }
 	virtual ALandscape* GetLandscapeActor();
 
 	/** Flush the foliage cache */
