Fixing NonUnity build

diff --git a/Engine/Source/Editor/WorldBrowser/Private/Tiles/WorldTileModel.cpp b/Engine/Source/Editor/WorldBrowser/Private/Tiles/WorldTileModel.cpp
--- a/Engine/Source/Editor/WorldBrowser/Private/Tiles/WorldTileModel.cpp
+++ b/Engine/Source/Editor/WorldBrowser/Private/Tiles/WorldTileModel.cpp
@@ -12,6 +12,7 @@
 #include "GameFramework/WorldSettings.h"
 #include "Engine/LevelStreaming.h"
 #include "Engine/LevelStreamingKismet.h"
+#include "Landscape.h"
 
 
 #define LOCTEXT_NAMESPACE "WorldBrowser"
