Fixed issues when renaming levels in the content browser which are referenced as streaming levels in the world.

diff --git a/Engine/Source/Editor/WorldBrowser/Private/LevelModel.cpp b/Engine/Source/Editor/WorldBrowser/Private/LevelModel.cpp
--- a/Engine/Source/Editor/WorldBrowser/Private/LevelModel.cpp
+++ b/Engine/Source/Editor/WorldBrowser/Private/LevelModel.cpp
@@ -8,6 +8,7 @@
 #include "LevelBrowserSettings.h"
 #include "Engine/Selection.h"
 #include "Engine/LevelStreaming.h"
+#include "AssetRegistryModule.h"
 
 #define LOCTEXT_NAMESPACE "WorldBrowser"
 
@@ -23,6 +24,26 @@ FLevelModel::FLevelModel(FLevelCollectionModel& InLevelCollectionModel,
 	, LevelActorsCount(0)
 {
 	SimulationStatus = FSimulationLevelStatus();
+
+	FAssetRegistryModule& AssetRegistryModule = FModuleManager::GetModuleChecked<FAssetRegistryModule>(TEXT("AssetRegistry"));
+	AssetRegistryModule.Get().OnAssetRenamed().AddRaw(this, &FLevelModel::OnAssetRenamed);
+}
+
+FLevelModel::~FLevelModel()
+{
+	FAssetRegistryModule& AssetRegistryModule = FModuleManager::GetModuleChecked<FAssetRegistryModule>(TEXT("AssetRegistry"));
+	AssetRegistryModule.Get().OnAssetRenamed().RemoveAll(this);
+}
+
+void FLevelModel::OnAssetRenamed(const FAssetData& AssetData, const FString& OldObjectPath)
+{
+	const FString CurrentPackage = GetLongPackageName().ToString();
+
+	if (FPackageName::ObjectPathToPackageName(OldObjectPath) == CurrentPackage)
+	{
+		UpdateAsset(AssetData);
+		UpdateDisplayName();
+	}
 }
 
 void FLevelModel::SetLevelSelectionFlag(bool bSelectedFlag)
diff --git a/Engine/Source/Editor/WorldBrowser/Private/LevelModel.h b/Engine/Source/Editor/WorldBrowser/Private/LevelModel.h
--- a/Engine/Source/Editor/WorldBrowser/Private/LevelModel.h
+++ b/Engine/Source/Editor/WorldBrowser/Private/LevelModel.h
@@ -29,7 +29,7 @@ public:
 	FLevelModel(FLevelCollectionModel& InLevelCollectionModel,
 				const TWeakObjectPtr< UEditorEngine >& InEditor);
 
-	virtual ~FLevelModel() {};
+	virtual ~FLevelModel();
 	
 	/** Traverses level model hierarchy */
 	void Accept(FLevelModelVisitor& Vistor);
@@ -72,6 +72,9 @@ public:
 
 	/**	@return	Level package file name */
 	virtual FName GetLongPackageName() const = 0;
+
+	/** Update asset associated with level model */
+	virtual void UpdateAsset(const FAssetData& AssetData) = 0;
 	
 	/** Refreshes cached data */
 	virtual void Update();
@@ -253,7 +256,11 @@ public:
 	
 	/** @return Class used for streaming this level */
 	virtual UClass* GetStreamingClass() const;
-	
+
+protected:
+	/** Called when the map asset is renamed */
+	void OnAssetRenamed(const FAssetData& AssetData, const FString& OldObjectPath);
+
 protected:
 	/** Level model display name */
 	FString								DisplayName;
diff --git a/Engine/Source/Editor/WorldBrowser/Private/StreamingLevels/StreamingLevelModel.cpp b/Engine/Source/Editor/WorldBrowser/Private/StreamingLevels/StreamingLevelModel.cpp
--- a/Engine/Source/Editor/WorldBrowser/Private/StreamingLevels/StreamingLevelModel.cpp
+++ b/Engine/Source/Editor/WorldBrowser/Private/StreamingLevels/StreamingLevelModel.cpp
@@ -65,6 +65,14 @@ FName FStreamingLevelModel::GetLongPackageName() const
 	}
 }
 
+void FStreamingLevelModel::UpdateAsset(const FAssetData& AssetData)
+{
+	if (LevelStreaming.IsValid())
+	{
+		LevelStreaming->SetWorldAssetByPackageName(AssetData.PackageName);
+	}
+}
+
 bool FStreamingLevelModel::SupportsLevelColor() const
 {
 	return LevelStreaming.IsValid();
diff --git a/Engine/Source/Editor/WorldBrowser/Private/StreamingLevels/StreamingLevelModel.h b/Engine/Source/Editor/WorldBrowser/Private/StreamingLevels/StreamingLevelModel.h
--- a/Engine/Source/Editor/WorldBrowser/Private/StreamingLevels/StreamingLevelModel.h
+++ b/Engine/Source/Editor/WorldBrowser/Private/StreamingLevels/StreamingLevelModel.h
@@ -34,6 +34,7 @@ public:
 	virtual ULevel* GetLevelObject() const override;
 	virtual FName GetAssetName() const override;
 	virtual FName GetLongPackageName() const override;
+	virtual void UpdateAsset(const FAssetData& AssetData) override;
 	virtual bool SupportsLevelColor() const override;
 	virtual FLinearColor GetLevelColor() const override;
 	virtual void SetLevelColor(FLinearColor InColor) override;
diff --git a/Engine/Source/Editor/WorldBrowser/Private/Tiles/WorldTileModel.cpp b/Engine/Source/Editor/WorldBrowser/Private/Tiles/WorldTileModel.cpp
--- a/Engine/Source/Editor/WorldBrowser/Private/Tiles/WorldTileModel.cpp
+++ b/Engine/Source/Editor/WorldBrowser/Private/Tiles/WorldTileModel.cpp
@@ -50,6 +50,7 @@ FWorldTileModel::FWorldTileModel(const TWeakObjectPtr<UEditorEngine>& InEditor,
 		TileDetails->bPersistentLevel = false;
 				
 		// Asset name for storing tile thumbnail inside package
+		SetAssetName(Tile.PackageName);
 		//FString AssetNameString = FString::Printf(TEXT("%s %s.TheWorld:PersistentLevel"), *ULevel::StaticClass()->GetName(), *Tile.PackageName.ToString());
 		FString AssetNameString = FString::Printf(TEXT("%s %s"), *UPackage::StaticClass()->GetName(), *Tile.PackageName.ToString());
 		AssetName =	FName(*AssetNameString);
@@ -115,6 +116,12 @@ bool FWorldTileModel::IsRootTile() const
 	return TileDetails->bPersistentLevel;
 }
 
+void FWorldTileModel::SetAssetName(const FName& PackageName)
+{
+	FString AssetNameString = FString::Printf(TEXT("%s %s"), *UPackage::StaticClass()->GetName(), *PackageName.ToString());
+	AssetName = FName(*AssetNameString);
+}
+
 FName FWorldTileModel::GetAssetName() const
 {
 	return AssetName;
@@ -125,6 +132,28 @@ FName FWorldTileModel::GetLongPackageName() const
 	return TileDetails->PackageName;
 }
 
+void FWorldTileModel::UpdateAsset(const FAssetData& AssetData)
+{
+	check(TileDetails != nullptr);
+	const FName OldPackageName = TileDetails->PackageName;
+
+	// Patch up any parent references which have been renamed
+	for (const TSharedPtr<FLevelModel>& LevelModel : LevelCollectionModel.GetAllLevels())
+	{
+		TSharedPtr<FWorldTileModel> WorldTileModel = StaticCastSharedPtr<FWorldTileModel>(LevelModel);
+
+		check(WorldTileModel->TileDetails != nullptr);
+		if (WorldTileModel->TileDetails->ParentPackageName == OldPackageName)
+		{
+			WorldTileModel->TileDetails->ParentPackageName = AssetData.PackageName;
+		}
+	}
+
+	const FName PackageName = AssetData.PackageName;
+	SetAssetName(PackageName);
+	TileDetails->PackageName = PackageName;
+}
+
 FVector2D FWorldTileModel::GetLevelPosition2D() const
 {
 	if (TileDetails->Bounds.IsValid)
diff --git a/Engine/Source/Editor/WorldBrowser/Private/Tiles/WorldTileModel.h b/Engine/Source/Editor/WorldBrowser/Private/Tiles/WorldTileModel.h
--- a/Engine/Source/Editor/WorldBrowser/Private/Tiles/WorldTileModel.h
+++ b/Engine/Source/Editor/WorldBrowser/Private/Tiles/WorldTileModel.h
@@ -73,6 +73,7 @@ public:
 	virtual FName GetAssetName() const override;
 	virtual FName GetLongPackageName() const override;
 	virtual void Update() override;
+	virtual void UpdateAsset(const FAssetData& AssetData) override;
 	virtual void LoadLevel() override;
 	virtual void SetVisible(bool bVisible) override;
 	virtual FVector2D GetLevelPosition2D() const override;
@@ -183,7 +184,10 @@ private:
 	
 	/** Handler for ZOrder chnages event from Tile details object  */
 	void OnZOrderPropertyChanged();
-	
+
+	/** Set the asset name based on the passed in package name */
+	void SetAssetName(const FName& PackageName);
+
 public:
 	/** This tile index in world composition tile list */
 	int32									TileIdx;
diff --git a/Engine/Source/Editor/WorldBrowser/WorldBrowser.Build.cs b/Engine/Source/Editor/WorldBrowser/WorldBrowser.Build.cs
--- a/Engine/Source/Editor/WorldBrowser/WorldBrowser.Build.cs
+++ b/Engine/Source/Editor/WorldBrowser/WorldBrowser.Build.cs
@@ -10,6 +10,7 @@ public class WorldBrowser : ModuleRules
 
         PrivateIncludePathModuleNames.AddRange(
         new string[] {
+                "AssetRegistry",
 				"AssetTools",
                 "MeshUtilities",
                 "ContentBrowser",
@@ -47,6 +48,7 @@ public class WorldBrowser : ModuleRules
 
         DynamicallyLoadedModuleNames.AddRange(
             new string[] {
+                "AssetRegistry",
 				"AssetTools",
 				"SceneOutliner",
                 "MeshUtilities",
