Added messages to explain what exactly is happening when user attempts to enable world composition in the WorldSettings.

diff --git a/Engine/Source/Editor/WorldBrowser/Private/WorldBrowserModule.cpp b/Engine/Source/Editor/WorldBrowser/Private/WorldBrowserModule.cpp
--- a/Engine/Source/Editor/WorldBrowser/Private/WorldBrowserModule.cpp
+++ b/Engine/Source/Editor/WorldBrowser/Private/WorldBrowserModule.cpp
@@ -29,8 +29,8 @@ void FWorldBrowserModule::StartupModule()
 		GEngine->OnWorldAdded().AddRaw(this, &FWorldBrowserModule::OnWorldCreated);
 		GEngine->OnWorldDestroyed().AddRaw(this, &FWorldBrowserModule::OnWorldDestroyed);
 	}
-	UWorldComposition::OnWorldCompositionCreated.AddRaw(this, &FWorldBrowserModule::OnWorldCompositionChanged);
-	UWorldComposition::OnWorldCompositionDestroyed.AddRaw(this, &FWorldBrowserModule::OnWorldCompositionChanged);
+
+	UWorldComposition::EnableWorldCompositionEvent.BindRaw(this, &FWorldBrowserModule::EnableWorldComposition);
 }
 
 void FWorldBrowserModule::ShutdownModule()
@@ -40,8 +40,8 @@ void FWorldBrowserModule::ShutdownModule()
 		GEngine->OnWorldAdded().RemoveAll(this);
 		GEngine->OnWorldDestroyed().RemoveAll(this);
 	}
-	UWorldComposition::OnWorldCompositionCreated.RemoveAll(this);
-	UWorldComposition::OnWorldCompositionDestroyed.RemoveAll(this);
+	
+	UWorldComposition::EnableWorldCompositionEvent.Unbind();
 
 	FLevelCollectionCommands::Unregister();
 
@@ -128,7 +128,70 @@ TSharedPtr<FLevelCollectionModel> FWorldBrowserModule::SharedWorldModel(UWorld*
 	return SharedWorldModel;
 }
 
+bool FWorldBrowserModule::EnableWorldComposition(UWorld* InWorld, bool bEnable)
+{
+	if (InWorld == nullptr || InWorld->WorldType != EWorldType::Editor)
+	{
+		return false;
+	}
+		
+	if (!bEnable)
+	{
+		if (InWorld->WorldComposition != nullptr)
+		{
+			InWorld->FlushLevelStreaming();
+			InWorld->WorldComposition->MarkPendingKill();
+			InWorld->WorldComposition = nullptr;
+			OnWorldCompositionChanged(InWorld);
+		}
 
+		return false;
+	}
+	
+	if (InWorld->WorldComposition == nullptr)
+	{
+		FString RootPackageName = InWorld->GetOutermost()->GetName();
 
+		// Map should be saved to disk
+		if (!FPackageName::DoesPackageExist(RootPackageName))
+		{
+			FMessageDialog::Open(EAppMsgType::Ok, LOCTEXT("EnableWorldCompositionNotSaved_Message", "Please save your level to disk before enabling World Composition"));
+			return false;
+		}
+			
+		// All existing sub-levels on this map should be removed
+		int32 NumExistingSublevels = InWorld->StreamingLevels.Num();
+		if (NumExistingSublevels > 0)
+		{
+			FMessageDialog::Open(EAppMsgType::Ok, LOCTEXT("EnableWorldCompositionExistingSublevels_Message", "World Composition cannot be enabled because there are already sub-levels manually added to the persistent level. World Composition uses auto-discovery so you must first remove any manually added sub-levels from the Levels window"));
+			return false;
+		}
+			
+		UWorldComposition* WorldCompostion = ConstructObject<UWorldComposition>(UWorldComposition::StaticClass(), InWorld);
+		// All map files found in the same and folder and all sub-folders will be added ass sub-levels to this map
+		// Make sure user understands this
+		int32 NumFoundSublevels = WorldCompostion->GetTilesList().Num();
+		if (NumFoundSublevels)
+		{
+			FFormatNamedArguments Arguments;
+			Arguments.Add(TEXT("NumSubLevels"), NumFoundSublevels);
+			Arguments.Add(TEXT("FolderLocation"), FText::FromString(FPackageName::GetLongPackagePath(RootPackageName)));
+			const FText Message = FText::Format(LOCTEXT("EnableWorldCompositionPrompt_Message", "World Composition auto-discovers sub-levels by scanning the folder the level is saved in, and all sub-folders. {NumSubLevels} level files were found in {FolderLocation} and will be added as sub-levels. Do you want to continue?"), Arguments);
+			
+			auto AppResult = FMessageDialog::Open(EAppMsgType::OkCancel, Message);
+			if (AppResult != EAppReturnType::Ok)
+			{
+				WorldCompostion->MarkPendingKill();
+				return false;
+			}
+		}
+			
+		// 
+		InWorld->WorldComposition = WorldCompostion;
+		OnWorldCompositionChanged(InWorld);
+	}
+	
+	return true;
+}
 
 #undef LOCTEXT_NAMESPACE
diff --git a/Engine/Source/Editor/WorldBrowser/Public/WorldBrowserModule.h b/Engine/Source/Editor/WorldBrowser/Public/WorldBrowserModule.h
--- a/Engine/Source/Editor/WorldBrowser/Public/WorldBrowserModule.h
+++ b/Engine/Source/Editor/WorldBrowser/Public/WorldBrowserModule.h
@@ -51,7 +51,8 @@ private:
 	void OnWorldCreated(UWorld* InWorld);
 	void OnWorldDestroyed(UWorld* InWorld);
 	void OnWorldCompositionChanged(UWorld* InWorld);
-		
+	bool EnableWorldComposition(UWorld* InWorld, bool bEnable);
+			
 private:
 	TWeakPtr<class FLevelCollectionModel>	WorldModel;
 };
diff --git a/Engine/Source/Runtime/Engine/Classes/Engine/WorldComposition.h b/Engine/Source/Runtime/Engine/Classes/Engine/WorldComposition.h
--- a/Engine/Source/Runtime/Engine/Classes/Engine/WorldComposition.h
+++ b/Engine/Source/Runtime/Engine/Classes/Engine/WorldComposition.h
@@ -142,11 +142,9 @@ class ENGINE_API UWorldComposition : public UObject
 	/** Collect tiles package names to cook  */
 	void CollectTilesToCook(TArray<FString>& PackageNames);
 
-	DECLARE_MULTICAST_DELEGATE_OneParam(FWorldCompositionEvent, UWorld*);
-	// Callback on world composition creation
-	static FWorldCompositionEvent OnWorldCompositionCreated;
-	// Callback on world composition destruction 
-	static FWorldCompositionEvent OnWorldCompositionDestroyed;
+	// Delegate to enable/disable world composition in the world
+	DECLARE_DELEGATE_RetVal_TwoParams(bool, FEnableWorldCompositionEvent, UWorld*, bool);
+	static FEnableWorldCompositionEvent EnableWorldCompositionEvent;
 
 #endif //WITH_EDITOR
 
diff --git a/Engine/Source/Runtime/Engine/Classes/GameFramework/WorldSettings.h b/Engine/Source/Runtime/Engine/Classes/GameFramework/WorldSettings.h
--- a/Engine/Source/Runtime/Engine/Classes/GameFramework/WorldSettings.h
+++ b/Engine/Source/Runtime/Engine/Classes/GameFramework/WorldSettings.h
@@ -208,8 +208,7 @@ class ENGINE_API AWorldSettings : public AInfo, public IInterface_AssetUserData
 
 	/** 
 	 * Enables tools for composing a tiled world. 
-	 * Level has to be saved to enable this option.
-	 * Warning: Enabling this option will remove all streaming levels from your persistent level.
+	 * Level has to be saved and all sub-levels removed before enabling this option.
 	 */
 	UPROPERTY(EditAnywhere, BlueprintReadOnly, Category=World)
 	uint32 bEnableWorldComposition:1;
@@ -458,8 +457,5 @@ private:
 	HIDE_ACTOR_TRANSFORM_FUNCTIONS();
 
 	virtual void Serialize( FArchive& Ar ) override;
-
-	/** Toggles world composition mode */
-	void EnabledWorldComposition(bool bEnable);
 };
 
diff --git a/Engine/Source/Runtime/Engine/Private/WorldComposition.cpp b/Engine/Source/Runtime/Engine/Private/WorldComposition.cpp
--- a/Engine/Source/Runtime/Engine/Private/WorldComposition.cpp
+++ b/Engine/Source/Runtime/Engine/Private/WorldComposition.cpp
@@ -11,8 +11,7 @@
 DEFINE_LOG_CATEGORY_STATIC(LogWorldComposition, Log, All);
 
 #if WITH_EDITOR
-UWorldComposition::FWorldCompositionEvent UWorldComposition::OnWorldCompositionCreated;
-UWorldComposition::FWorldCompositionEvent UWorldComposition::OnWorldCompositionDestroyed;
+UWorldComposition::FEnableWorldCompositionEvent UWorldComposition::EnableWorldCompositionEvent;
 #endif // WITH_EDITOR
 
 UWorldComposition::UWorldComposition(const FObjectInitializer& ObjectInitializer)
diff --git a/Engine/Source/Runtime/Engine/Private/WorldSettings.cpp b/Engine/Source/Runtime/Engine/Private/WorldSettings.cpp
--- a/Engine/Source/Runtime/Engine/Private/WorldSettings.cpp
+++ b/Engine/Source/Runtime/Engine/Private/WorldSettings.cpp
@@ -231,7 +231,14 @@ void AWorldSettings::PostEditChangeProperty(FPropertyChangedEvent& PropertyChang
 
 		if (PropertyThatChanged->GetName()==TEXT("bEnableWorldComposition"))
 		{
-			EnabledWorldComposition(bEnableWorldComposition);
+			if (UWorldComposition::EnableWorldCompositionEvent.IsBound())
+			{
+				bEnableWorldComposition = UWorldComposition::EnableWorldCompositionEvent.Execute(GetWorld(), bEnableWorldComposition);
+			}
+			else
+			{
+				bEnableWorldComposition = false;
+			}
 		}
 	}
 
@@ -263,38 +270,6 @@ void AWorldSettings::PostEditChangeProperty(FPropertyChangedEvent& PropertyChang
 	Super::PostEditChangeProperty(PropertyChangedEvent);
 }
 
-void AWorldSettings::EnabledWorldComposition(bool bEnable)
-{
-	bEnableWorldComposition = false;
-
-	if (!bEnable)
-	{
-		if (GetWorld()->WorldComposition != nullptr)
-		{
-			GetWorld()->WorldComposition = nullptr;
-			GetWorld()->StreamingLevels.Empty();
-			GetWorld()->FlushLevelStreaming();
-
-			UWorldComposition::OnWorldCompositionDestroyed.Broadcast(GetWorld());
-		}
-	}
-	else
-	{
-		FString RootPackageName = GetOutermost()->GetName();
-
-		if (GetWorld()->WorldComposition == nullptr && FPackageName::DoesPackageExist(RootPackageName))
-		{
-			GetWorld()->StreamingLevels.Empty();
-			GetWorld()->FlushLevelStreaming();
-			GetWorld()->WorldComposition = ConstructObject<UWorldComposition>(UWorldComposition::StaticClass(), GetWorld());
-			
-			bEnableWorldComposition = true;
-
-			UWorldComposition::OnWorldCompositionCreated.Broadcast(GetWorld());
-		}
-	}
-}
-
 #endif // WITH_EDITOR
 
 #undef LOCTEXT_NAMESPACE
