Fixed tiled landscape importing regression, section offset was not correcly assigned to a new landscape proxies.

diff --git a/Engine/Source/Editor/WorldBrowser/Private/Tiles/WorldTileCollectionModel.cpp b/Engine/Source/Editor/WorldBrowser/Private/Tiles/WorldTileCollectionModel.cpp
--- a/Engine/Source/Editor/WorldBrowser/Private/Tiles/WorldTileCollectionModel.cpp
+++ b/Engine/Source/Editor/WorldBrowser/Private/Tiles/WorldTileCollectionModel.cpp
@@ -1627,16 +1627,13 @@ void FWorldTileCollectionModel::ImportTiledLandscape_Executed()
 
 					if (NewLandscape)
 					{
-						// Set landscape actor location at min corner of bounding rect, so it will be centered around origin
-						FIntRect NewLandscapeRect = NewLandscape->GetBoundingRect();
-						float WidthX = NewLandscapeRect.Width()*TileScale.X;
-						float WidthY = NewLandscapeRect.Height()*TileScale.Y;
-						NewLandscape->SetActorLocation(-FVector(WidthX, WidthY, 0.f)*0.5f);
-						
 						// Set bounds of a tile
 						NewTileModel->TileDetails->Bounds = NewLandscape->GetComponentsBoundingBox();
-
+						
 						// Calculate this tile offset from world origin
+						FIntRect NewLandscapeRect = NewLandscape->GetBoundingRect();
+						float WidthX = NewLandscapeRect.Width()*TileScale.X;
+						float WidthY = NewLandscapeRect.Height()*TileScale.Y;
 						FIntPoint TileCoordinates = ImportSettings.TileCoordinates[TileIndex] + ImportSettings.TilesCoordinatesOffset;
 						FIntPoint TileOffset = FIntPoint(TileCoordinates.X*WidthX, TileCoordinates.Y*WidthY);
 						
diff --git a/Engine/Source/Editor/WorldBrowser/Private/Tiles/WorldTileModel.cpp b/Engine/Source/Editor/WorldBrowser/Private/Tiles/WorldTileModel.cpp
--- a/Engine/Source/Editor/WorldBrowser/Private/Tiles/WorldTileModel.cpp
+++ b/Engine/Source/Editor/WorldBrowser/Private/Tiles/WorldTileModel.cpp
@@ -473,6 +473,8 @@ void FWorldTileModel::SetLevelPosition(const FIntPoint& InPosition)
 	if (IsLandscapeBased())
 	{
 		UpdateLandscapeSectionsOffset(Offset);
+		bool bShowWarnings = true;
+		ULandscapeInfo::RecreateLandscapeInfo(LevelCollectionModel.GetWorld(), bShowWarnings);
 	}
 	
 	// Transform child levels
@@ -848,6 +850,9 @@ ALandscapeProxy* FWorldTileModel::ImportLandscapeTile(const FLandscapeImportSett
 		LandscapeProxy->LandscapeMaterial = Settings.LandscapeMaterial;
 	}
 	
+	// Cache pointer to landscape in the level model
+	Landscape = LandscapeProxy;
+
 	// Create landscape components	
 	LandscapeProxy->Import(
 		Settings.LandscapeGuid, 
