Added command to 'Levels browser' to move selected foliage instances to specific sub-level
UE-9430

diff --git a/Engine/Source/Editor/UnrealEd/Classes/Editor/EditorEngine.h b/Engine/Source/Editor/UnrealEd/Classes/Editor/EditorEngine.h
--- a/Engine/Source/Editor/UnrealEd/Classes/Editor/EditorEngine.h
+++ b/Engine/Source/Editor/UnrealEd/Classes/Editor/EditorEngine.h
@@ -1809,6 +1809,13 @@ public:
 	 */
 	void MoveSelectedActorsToLevel( ULevel* InLevel );
 
+	/**
+	 * Moves selected foliage instances to the target level.
+	 *
+	 * @param	InLevel		The target level.
+	 */
+	void MoveSelectedFoliageToLevel( ULevel* InTargetLevel );
+
 	/**
 	 * Checks to see whether it's possible to perform a copy operation on the selected actors.
 	 *
diff --git a/Engine/Source/Editor/UnrealEd/Private/EditorServer.cpp b/Engine/Source/Editor/UnrealEd/Private/EditorServer.cpp
--- a/Engine/Source/Editor/UnrealEd/Private/EditorServer.cpp
+++ b/Engine/Source/Editor/UnrealEd/Private/EditorServer.cpp
@@ -2833,6 +2833,38 @@ void UEditorEngine::DoMoveSelectedActorsToLevel( ULevel* InDestLevel )
 	GEditor->Trans->End();
 }
 
+void UEditorEngine::MoveSelectedFoliageToLevel(ULevel* InTargetLevel)
+{
+	// Can't move into a locked level
+	if (FLevelUtils::IsLevelLocked(InTargetLevel))
+	{
+		FNotificationInfo Info(NSLOCTEXT("UnrealEd", "CannotMoveFoliageIntoLockedLevel", "Cannot move the selected foliage into a locked level"));
+		Info.bUseThrobber = false;
+		FSlateNotificationManager::Get().AddNotification(Info)->SetCompletionState(SNotificationItem::CS_Fail);
+		return;
+	}
+
+	const FScopedTransaction Transaction(NSLOCTEXT("UnrealEd", "MoveSelectedFoliageToSelectedLevel", "Move Selected Foliage to Level"));
+
+	// Get a world context
+	UWorld* World = InTargetLevel->OwningWorld;
+	
+	// Iterate over all foliage actors in the world and move selected instances to a foliage actor in the target level
+	const int32 NumLevels = World->GetNumLevels();
+	for (int32 LevelIdx = 0; LevelIdx < NumLevels; ++LevelIdx)
+	{
+		ULevel* Level = World->GetLevel(LevelIdx);
+		if (Level != InTargetLevel)
+		{
+			AInstancedFoliageActor* IFA = AInstancedFoliageActor::GetInstancedFoliageActorForLevel(Level, /*bCreateIfNone*/ false);
+			if (IFA && IFA->HasSelectedInstances())
+			{
+				IFA->MoveSelectedInstancesToLevel(InTargetLevel);
+			}
+		}
+	}
+}
+
 ULevel*  UEditorEngine::CreateTransLevelMoveBuffer( UWorld* InWorld )
 {
 	ULevel* BufferLevel = NewObject<ULevel>(GetTransientPackage(), TEXT("TransLevelMoveBuffer"));
diff --git a/Engine/Source/Editor/WorldBrowser/Private/LevelCollectionCommands.h b/Engine/Source/Editor/WorldBrowser/Private/LevelCollectionCommands.h
--- a/Engine/Source/Editor/WorldBrowser/Private/LevelCollectionCommands.h
+++ b/Engine/Source/Editor/WorldBrowser/Private/LevelCollectionCommands.h
@@ -75,6 +75,7 @@ public:
 		UI_COMMAND( AddsActors, "Select Actors", "Adds the Actors in the selected Levels from the viewport's existing selection", EUserInterfaceActionType::Button, FInputGesture() );
 		UI_COMMAND( RemovesActors, "Deselect Actors", "Removes the Actors in the selected Levels from the viewport's existing selection", EUserInterfaceActionType::Button, FInputGesture() );
 		UI_COMMAND( MoveActorsToSelected, "Move Selected Actors to Level", "Moves the selected actors to this level", EUserInterfaceActionType::Button, FInputGesture() );
+		UI_COMMAND( MoveFoliageToSelected, "Move Selected Foliage to Level", "Moves the selected foliage instances to this level. Keeps cross-level references to original bases", EUserInterfaceActionType::Button, FInputGesture() );
 		UI_COMMAND( SelectStreamingVolumes, "Select Streaming Volumes", "Selects the streaming volumes associated with the selected levels", EUserInterfaceActionType::Button, FInputGesture() );
 
 		//visibility
@@ -227,6 +228,9 @@ public:
 	/** Moves the selected actors to the selected level */
 	TSharedPtr< FUICommandInfo > MoveActorsToSelected;
 
+	/** Moves the selected foliage instances to the selected level */
+	TSharedPtr< FUICommandInfo > MoveFoliageToSelected;
+	
 	/** Selects the streaming volumes associated with the selected levels */
 	TSharedPtr< FUICommandInfo > SelectStreamingVolumes;
 	
diff --git a/Engine/Source/Editor/WorldBrowser/Private/LevelCollectionModel.cpp b/Engine/Source/Editor/WorldBrowser/Private/LevelCollectionModel.cpp
--- a/Engine/Source/Editor/WorldBrowser/Private/LevelCollectionModel.cpp
+++ b/Engine/Source/Editor/WorldBrowser/Private/LevelCollectionModel.cpp
@@ -85,6 +85,10 @@ void FLevelCollectionModel::BindCommands()
 	ActionList.MapAction(Commands.MoveActorsToSelected,
 		FExecuteAction::CreateSP(this, &FLevelCollectionModel::MoveActorsToSelected_Executed),
 		FCanExecuteAction::CreateSP(this, &FLevelCollectionModel::IsValidMoveActorsToLevel));
+
+	ActionList.MapAction(Commands.MoveFoliageToSelected,
+		FExecuteAction::CreateSP(this, &FLevelCollectionModel::MoveFoliageToSelected_Executed),
+		FCanExecuteAction::CreateSP(this, &FLevelCollectionModel::IsValidMoveFoliageToLevel));
 		
 	ActionList.MapAction(Commands.World_SaveSelectedLevels,
 		FExecuteAction::CreateSP(this, &FLevelCollectionModel::SaveSelectedLevels_Executed),
@@ -1505,6 +1509,15 @@ void FLevelCollectionModel::MoveActorsToSelected_Executed()
 	RequestUpdateAllLevels();
 }
 
+void FLevelCollectionModel::MoveFoliageToSelected_Executed()
+{
+	if (GetSelectedLevels().Num() == 1)
+	{
+		ULevel* TargetLevel = GetSelectedLevels()[0]->GetLevelObject();
+		Editor->MoveSelectedFoliageToLevel(TargetLevel);
+	}
+}
+
 void FLevelCollectionModel::SelectActors_Executed()
 {
 	//first clear any existing actor selection
@@ -1763,7 +1776,7 @@ void FLevelCollectionModel::CacheCanExecuteSourceControlVars() const
 	}
 }
 
-bool FLevelCollectionModel::IsValidMoveActorsToLevel() 
+bool FLevelCollectionModel::IsValidMoveActorsToLevel() const
 {
 	static bool bCachedIsValidActorMoveResult = false;
 	if (bSelectionHasChanged)
@@ -1799,6 +1812,18 @@ bool FLevelCollectionModel::IsValidMoveActorsToLevel()
 	return bCachedIsValidActorMoveResult && AreAllSelectedLevelsEditableAndVisible();
 }
 
+bool FLevelCollectionModel::IsValidMoveFoliageToLevel() const
+{
+	if (IsOneLevelSelected() && 
+		AreAllSelectedLevelsEditableAndVisible() && 
+		GLevelEditorModeTools().IsModeActive(FBuiltinEditorModes::EM_Foliage))
+	{
+		return true;
+	}
+
+	return false;
+}
+
 void FLevelCollectionModel::OnActorSelectionChanged(UObject* obj)
 {
 	OnActorOrLevelSelectionChanged();
diff --git a/Engine/Source/Editor/WorldBrowser/Private/LevelCollectionModel.h b/Engine/Source/Editor/WorldBrowser/Private/LevelCollectionModel.h
--- a/Engine/Source/Editor/WorldBrowser/Private/LevelCollectionModel.h
+++ b/Engine/Source/Editor/WorldBrowser/Private/LevelCollectionModel.h
@@ -212,7 +212,10 @@ public:
 	bool AreActorsSelected() const;
 
 	/** @return whether moving the selected actors to the selected level is a valid action */
-	bool IsValidMoveActorsToLevel();
+	bool IsValidMoveActorsToLevel() const;
+
+	/** @return whether moving the selected foliage to the selected level is a valid action */
+	bool IsValidMoveFoliageToLevel() const;
 
 	/** delegate used to pickup when the selection has changed */
 	void OnActorSelectionChanged(UObject* obj);
@@ -307,6 +310,9 @@ protected:
 	/** Moves the selected actors to this level */
 	void MoveActorsToSelected_Executed();
 
+	/** Moves the selected foliage to this level */
+	void MoveFoliageToSelected_Executed();
+
 	/** Saves selected levels */
 	void SaveSelectedLevels_Executed();
 
@@ -523,7 +529,7 @@ protected:
 	mutable bool						bCanExecuteSCC;
 
 	/** Flag for whether the selection of levels or actors has changed */
-	bool								bSelectionHasChanged;
+	mutable bool						bSelectionHasChanged;
 
 	/** Guard to avoid recursive level selection updates */
 	bool								bUpdatingLevelsSelection;
diff --git a/Engine/Source/Editor/WorldBrowser/Private/StreamingLevels/StreamingLevelCollectionModel.cpp b/Engine/Source/Editor/WorldBrowser/Private/StreamingLevels/StreamingLevelCollectionModel.cpp
--- a/Engine/Source/Editor/WorldBrowser/Private/StreamingLevels/StreamingLevelCollectionModel.cpp
+++ b/Engine/Source/Editor/WorldBrowser/Private/StreamingLevels/StreamingLevelCollectionModel.cpp
@@ -302,6 +302,7 @@ void FStreamingLevelCollectionModel::BuildHierarchyMenu(FMenuBuilder& InMenuBuil
 		if (IsOneLevelSelected())
 		{
 			InMenuBuilder.AddMenuEntry( Commands.MoveActorsToSelected );
+			InMenuBuilder.AddMenuEntry( Commands.MoveFoliageToSelected );
 		}
 
 		if (AreAnyLevelsSelected() && !(IsOneLevelSelected() && SelectedLevelsList[0]->IsPersistent()))
diff --git a/Engine/Source/Editor/WorldBrowser/Private/Tiles/WorldTileCollectionModel.cpp b/Engine/Source/Editor/WorldBrowser/Private/Tiles/WorldTileCollectionModel.cpp
--- a/Engine/Source/Editor/WorldBrowser/Private/Tiles/WorldTileCollectionModel.cpp
+++ b/Engine/Source/Editor/WorldBrowser/Private/Tiles/WorldTileCollectionModel.cpp
@@ -358,6 +358,7 @@ void FWorldTileCollectionModel::BuildWorldCompositionMenu(FMenuBuilder& InMenuBu
 		if (IsOneLevelSelected())
 		{
 			InMenuBuilder.AddMenuEntry(Commands.MoveActorsToSelected);
+			InMenuBuilder.AddMenuEntry(Commands.MoveFoliageToSelected);
 		}
 	}
 	InMenuBuilder.EndSection();
@@ -468,6 +469,7 @@ void FWorldTileCollectionModel::BuildHierarchyMenu(FMenuBuilder& InMenuBuilder)
 		if (IsOneLevelSelected())
 		{
 			InMenuBuilder.AddMenuEntry( Commands.MoveActorsToSelected );
+			InMenuBuilder.AddMenuEntry( Commands.MoveFoliageToSelected );
 		}
 	}
 	InMenuBuilder.EndSection();
diff --git a/Engine/Source/Runtime/Foliage/Private/InstancedFoliage.cpp b/Engine/Source/Runtime/Foliage/Private/InstancedFoliage.cpp
--- a/Engine/Source/Runtime/Foliage/Private/InstancedFoliage.cpp
+++ b/Engine/Source/Runtime/Foliage/Private/InstancedFoliage.cpp
@@ -1339,6 +1339,42 @@ void AInstancedFoliageActor::MoveInstancesToNewComponent(UPrimitiveComponent* In
 	}
 }
 
+void AInstancedFoliageActor::MoveSelectedInstancesToLevel(ULevel* InTargetLevel)
+{
+	if (InTargetLevel == GetLevel() || !HasSelectedInstances())
+	{
+		return;
+	}
+		
+	AInstancedFoliageActor* TargetIFA = GetInstancedFoliageActorForLevel(InTargetLevel, /*bCreateIfNone*/ true);
+	
+	Modify();
+	TargetIFA->Modify();
+	
+	// Do move
+	for (auto& MeshPair : FoliageMeshes)
+	{
+		FFoliageMeshInfo& MeshInfo = *MeshPair.Value;
+		UFoliageType* FoliageType = MeshPair.Key;
+
+		if (MeshInfo.SelectedIndices.Num())
+		{
+			FFoliageMeshInfo* TargetMeshInfo = nullptr;
+			UFoliageType* TargetFoliageType = TargetIFA->AddFoliageType(FoliageType, &TargetMeshInfo);
+
+			// Add selected instances to the target actor
+			for (int32 InstanceIndex : MeshInfo.SelectedIndices)
+			{
+				FFoliageInstance& Instance = MeshInfo.Instances[InstanceIndex];
+				TargetMeshInfo->AddInstance(TargetIFA, TargetFoliageType, Instance, InstanceBaseCache.GetInstanceBasePtr(Instance.BaseId).Get());
+			}
+
+			// Remove selected instances from this actor
+			MeshInfo.RemoveInstances(this, MeshInfo.SelectedIndices.Array());
+		}
+	}
+}
+
 TMap<UFoliageType*, TArray<const FFoliageInstancePlacementInfo*>> AInstancedFoliageActor::GetInstancesForComponent(UActorComponent* InComponent)
 {
 	TMap<UFoliageType*, TArray<const FFoliageInstancePlacementInfo*>> Result;
@@ -1620,6 +1656,19 @@ void AInstancedFoliageActor::SelectInstance(UInstancedStaticMeshComponent* InCom
 	}
 }
 
+bool AInstancedFoliageActor::HasSelectedInstances() const
+{
+	for (const auto& MeshPair : FoliageMeshes)
+	{
+		if (MeshPair.Value->SelectedIndices.Num() > 0)
+		{
+			return true;
+		}
+	}
+	
+	return false;
+}
+
 void AInstancedFoliageActor::PostEditUndo()
 {
 	Super::PostEditUndo();
diff --git a/Engine/Source/Runtime/Foliage/Public/InstancedFoliageActor.h b/Engine/Source/Runtime/Foliage/Public/InstancedFoliageActor.h
--- a/Engine/Source/Runtime/Foliage/Public/InstancedFoliageActor.h
+++ b/Engine/Source/Runtime/Foliage/Public/InstancedFoliageActor.h
@@ -100,6 +100,9 @@ public:
 	// The instances keep the same world locations
 	FOLIAGE_API void MoveInstancesToNewComponent(UPrimitiveComponent* InOldComponent, UPrimitiveComponent* InNewComponent);
 
+	// Move selected instances to a foliage actor in target level
+	FOLIAGE_API void MoveSelectedInstancesToLevel(ULevel* InTargetLevel);
+	
 	// Move instances based on a component that has just been moved.
 	void MoveInstancesForMovedComponent(UActorComponent* InComponent);
 	
@@ -128,6 +131,9 @@ public:
 	// Select an individual instance.
 	FOLIAGE_API void SelectInstance(UInstancedStaticMeshComponent* InComponent, int32 InComponentInstanceIndex, bool bToggle);
 
+	// Whether actor has selected instances
+	FOLIAGE_API bool HasSelectedInstances() const;
+
 	// Propagate the selected instances to the actual render components
 	FOLIAGE_API void ApplySelectionToComponents(bool bApply);
 
