Fixed: Import Tiled Landscape Material dropdown look not consistent
UE-8177

diff --git a/Engine/Source/Editor/WorldBrowser/Private/Tiles/STiledLandscapeImportDlg.cpp b/Engine/Source/Editor/WorldBrowser/Private/Tiles/STiledLandscapeImportDlg.cpp
--- a/Engine/Source/Editor/WorldBrowser/Private/Tiles/STiledLandscapeImportDlg.cpp
+++ b/Engine/Source/Editor/WorldBrowser/Private/Tiles/STiledLandscapeImportDlg.cpp
@@ -8,6 +8,7 @@
 #include "ContentBrowserModule.h"
 #include "LandscapeProxy.h"
 #include "SNumericEntryBox.h"
+#include "PropertyCustomizationHelpers.h"
 
 #define LOCTEXT_NAMESPACE "WorldBrowser"
 
@@ -170,18 +171,11 @@ void STiledLandcapeImportDlg::Construct(const FArguments& InArgs, TSharedPtr<SWi
 				+SUniformGridPanel::Slot(1, 4)
 				.VAlign(VAlign_Center)
 				[
-					SAssignNew(LandscapeMaterialComboButton, SComboButton)
-					.VAlign(EVerticalAlignment::VAlign_Center)
-					.ComboButtonStyle(FEditorStyle::Get(), "ToolbarComboButton")
-					.ForegroundColor(FLinearColor::White)
-					.ContentPadding(3)
-					.MenuPlacement(EMenuPlacement::MenuPlacement_BelowAnchor)
-					.OnGetMenuContent(this, &STiledLandcapeImportDlg::CreateLandscapeMaterialPicker)
-					.ButtonContent()
-					[
-						SNew( STextBlock )
-						.Text(this, &STiledLandcapeImportDlg::GetLandscapeMaterialName)
-					]
+					SNew(SObjectPropertyEntryBox)
+					.AllowedClass(UMaterialInterface::StaticClass())
+					.ObjectPath(this, &STiledLandcapeImportDlg::GetLandscapeMaterialPath)
+					.OnObjectChanged(this, &STiledLandcapeImportDlg::OnLandscapeMaterialChanged)
+					.AllowClear(true)
 				]
 			]
 
@@ -245,36 +239,6 @@ void STiledLandcapeImportDlg::Construct(const FArguments& InArgs, TSharedPtr<SWi
 	SetPossibleConfigurationsForFileSize(0);
 }
 
-TSharedRef<SWidget> STiledLandcapeImportDlg::CreateLandscapeMaterialPicker()
-{
-	FContentBrowserModule& ContentBrowserModule = FModuleManager::Get().LoadModuleChecked<FContentBrowserModule>(TEXT("ContentBrowser"));
-
-	FAssetPickerConfig AssetPickerConfig;
-	AssetPickerConfig.Filter.bRecursiveClasses = true;
-	AssetPickerConfig.Filter.ClassNames.Add(UMaterial::StaticClass()->GetFName());
-	AssetPickerConfig.Filter.ClassNames.Add(UMaterialInstance::StaticClass()->GetFName());
-	AssetPickerConfig.OnAssetSelected = FOnAssetSelected::CreateSP(this, &STiledLandcapeImportDlg::OnLandscapeMaterialChanged);
-	AssetPickerConfig.InitialAssetViewType = EAssetViewType::List;
-	AssetPickerConfig.bAllowNullSelection = true;
-	
-	return SNew(SBox)
-		.WidthOverride(250)
-		.HeightOverride(300)
-		[
-			ContentBrowserModule.Get().CreateAssetPicker(AssetPickerConfig)
-		];
-}
-
-FText STiledLandcapeImportDlg::GetLandscapeMaterialName() const
-{
-	if (ImportSettings.LandscapeMaterial.IsValid())
-	{
-		return FText::FromString(ImportSettings.LandscapeMaterial->GetName());
-	}
-
-	return LOCTEXT("TiledLandscapeImport_NoLandscapeMaterialText", "None");
-}
-
 TSharedRef<SWidget> STiledLandcapeImportDlg::HandleTileConfigurationComboBoxGenarateWidget(TSharedPtr<FTileImportConfiguration> InItem) const
 {
 	const FText ItemText = GenerateConfigurationText(InItem->NumComponents, InItem->NumSectionsPerComponent, InItem->NumQuadsPerSection);
@@ -588,10 +552,14 @@ FReply STiledLandcapeImportDlg::OnClickedCancel()
 	return FReply::Handled();	
 }
 
+FString STiledLandcapeImportDlg::GetLandscapeMaterialPath() const
+{
+	return ImportSettings.LandscapeMaterial.IsValid() ? ImportSettings.LandscapeMaterial->GetPathName() : FString("");
+}
+
 void STiledLandcapeImportDlg::OnLandscapeMaterialChanged(const FAssetData& AssetData)
 {
 	ImportSettings.LandscapeMaterial = Cast<UMaterialInterface>(AssetData.GetAsset());
-	LandscapeMaterialComboButton->SetIsOpen(false);
 
 	// pull landscape layers from a chosen material
 	UpdateLandscapeLayerList();
diff --git a/Engine/Source/Editor/WorldBrowser/Private/Tiles/STiledLandscapeImportDlg.h b/Engine/Source/Editor/WorldBrowser/Private/Tiles/STiledLandscapeImportDlg.h
--- a/Engine/Source/Editor/WorldBrowser/Private/Tiles/STiledLandscapeImportDlg.h
+++ b/Engine/Source/Editor/WorldBrowser/Private/Tiles/STiledLandscapeImportDlg.h
@@ -30,10 +30,6 @@ public:
 	const FTiledLandscapeImportSettings& GetImportSettings() const;
 
 private:
-	/** */
-	TSharedRef<SWidget> CreateLandscapeMaterialPicker();
-	FText GetLandscapeMaterialName() const;
-
 	/** */
 	TSharedRef<SWidget> HandleTileConfigurationComboBoxGenarateWidget(TSharedPtr<FTileImportConfiguration> InItem) const;
 	FText				GetTileConfigurationText() const;
@@ -66,6 +62,7 @@ private:
 	FReply OnClickedCancel();
 
 	/** */
+	FString GetLandscapeMaterialPath() const;
 	void OnLandscapeMaterialChanged(const FAssetData& AssetData);
 
 	/** */
@@ -101,9 +98,6 @@ private:
 	/** */
 	TSharedPtr<SComboBox<TSharedPtr<FTileImportConfiguration>>> TileConfigurationComboBox;
 	
-	/** */
-	TSharedPtr<SComboButton> LandscapeMaterialComboButton;
-
 	TSharedPtr<SListView<TSharedPtr<FTiledLandscapeImportSettings::LandscapeLayerSettings>>>	LayerDataListView;
 	TArray<TSharedPtr<FTiledLandscapeImportSettings::LandscapeLayerSettings>>					LayerDataList;		
 		
