Fixed: 'Add Adjacent Ladscape' command in world composition was generating overlapped landscape components warnings

diff --git a/Engine/Source/Editor/WorldBrowser/Private/Tiles/WorldTileModel.cpp b/Engine/Source/Editor/WorldBrowser/Private/Tiles/WorldTileModel.cpp
--- a/Engine/Source/Editor/WorldBrowser/Private/Tiles/WorldTileModel.cpp
+++ b/Engine/Source/Editor/WorldBrowser/Private/Tiles/WorldTileModel.cpp
@@ -784,14 +784,15 @@ bool FWorldTileModel::CreateAdjacentLandscapeProxy(ALandscapeProxy* SourceLandsc
 	FIntPoint SourceLandscapeSize = SourceLandscapeRect.Size();
 
 	FLandscapeImportSettings ImportSettings = {};
-	ImportSettings.LandscapeGuid = SourceLandscape->GetLandscapeGuid();
-	ImportSettings.ComponentSizeQuads = SourceLandscape->ComponentSizeQuads;
+	ImportSettings.LandscapeGuid		= SourceLandscape->GetLandscapeGuid();
+	ImportSettings.LandscapeMaterial	= SourceLandscape->GetLandscapeMaterial();
+	ImportSettings.ComponentSizeQuads	= SourceLandscape->ComponentSizeQuads;
+	ImportSettings.QuadsPerSection		= SourceLandscape->SubsectionSizeQuads;
 	ImportSettings.SectionsPerComponent = SourceLandscape->NumSubsections;
-	ImportSettings.QuadsPerSection = SourceLandscape->SubsectionSizeQuads;
-	ImportSettings.SizeX = SourceLandscapeRect.Width() + 1;
-	ImportSettings.SizeY = SourceLandscapeRect.Height() + 1;
+	ImportSettings.SizeX				= SourceLandscapeRect.Width() + 1;
+	ImportSettings.SizeY				= SourceLandscapeRect.Height() + 1;
 
-	// Initialize blank heightmap data
+	// Initialize with blank heightmap data
 	ImportSettings.HeightData.AddUninitialized(ImportSettings.SizeX * ImportSettings.SizeY);
 	for (auto& HeightSample : ImportSettings.HeightData)
 	{
@@ -799,18 +800,18 @@ bool FWorldTileModel::CreateAdjacentLandscapeProxy(ALandscapeProxy* SourceLandsc
 	}
 
 	// Set proxy location at landscape bounds Min point
-	ImportSettings.LandscapeTransform.SetLocation(-FVector(SourceLandscapeSize, 0.f)*SourceLandscapeScale*0.5f + FVector(0.f, 0.f, SourceLandscape->GetActorLocation().Z));
+	ImportSettings.LandscapeTransform.SetLocation(FVector(0.f, 0.f, SourceLandscape->GetActorLocation().Z));
 	ImportSettings.LandscapeTransform.SetScale3D(SourceLandscapeScale);
 	
 	// Create new landscape object
-	ALandscapeProxy* AdjacenLandscape = ImportLandscapeTile(ImportSettings);
-	if (AdjacenLandscape)
+	ALandscapeProxy* AdjacentLandscape = ImportLandscapeTile(ImportSettings);
+	if (AdjacentLandscape)
 	{
 		// Copy source landscape properties 
-		AdjacenLandscape->GetSharedProperties(SourceLandscape);
+		AdjacentLandscape->GetSharedProperties(SourceLandscape);
 		
 		// Refresh level model bounding box
-		FBox AdjacentLandscapeBounds = AdjacenLandscape->GetComponentsBoundingBox(true);
+		FBox AdjacentLandscapeBounds = AdjacentLandscape->GetComponentsBoundingBox(true);
 		TileDetails->Bounds = AdjacentLandscapeBounds;
 
 		// Calculate proxy offset from source landscape actor
