Using correct landscape LOD for spawning proxy mesh during map simplification

diff --git a/Engine/Source/Editor/WorldBrowser/Private/Tiles/WorldTileCollectionModel.cpp b/Engine/Source/Editor/WorldBrowser/Private/Tiles/WorldTileCollectionModel.cpp
--- a/Engine/Source/Editor/WorldBrowser/Private/Tiles/WorldTileCollectionModel.cpp
+++ b/Engine/Source/Editor/WorldBrowser/Private/Tiles/WorldTileCollectionModel.cpp
@@ -1925,15 +1925,17 @@ bool FWorldTileCollectionModel::GenerateLODLevels(FLevelModelList InLevelList, i
 	
 		struct FAssetToSpawnInfo
 		{
-			FAssetToSpawnInfo(UStaticMesh* InStaticMesh, const FTransform& InTransform, ALandscapeProxy* InSourceLandscape = nullptr)
+			FAssetToSpawnInfo(UStaticMesh* InStaticMesh, const FTransform& InTransform, ALandscapeProxy* InSourceLandscape = nullptr, int32 InLandscapeLOD = 0)
 			: StaticMesh(InStaticMesh)
 			, Transform(InTransform)
 			, SourceLandscape(InSourceLandscape)
+			, LandscapeLOD(InLandscapeLOD)
 			{}
 
 			UStaticMesh* StaticMesh;
 			FTransform Transform;
 			ALandscapeProxy* SourceLandscape;
+			int32 LandscapeLOD;
 		};
 
 		TArray<FAssetToSpawnInfo>	AssetsToSpawn;
@@ -2087,7 +2089,7 @@ bool FWorldTileCollectionModel::GenerateLODLevels(FLevelModelList InLevelList, i
 			}
 
 			GeneratedAssets.Add(StaticMesh);
-			new(AssetsToSpawn) FAssetToSpawnInfo(StaticMesh, FTransform(LandscapeWorldLocation - ActorsOffset), Landscape);
+			new(AssetsToSpawn) FAssetToSpawnInfo(StaticMesh, FTransform(LandscapeWorldLocation - ActorsOffset), Landscape, LandscapeLOD);
 
 			LandscapeActorIndex++;
 		}
@@ -2143,7 +2145,7 @@ bool FWorldTileCollectionModel::GenerateLODLevels(FLevelModelList InLevelList, i
 				{
 					ALandscapeMeshProxyActor* MeshActor = LODWorld->SpawnActor<ALandscapeMeshProxyActor>(Location, Rotation);
 					MeshActor->GetLandscapeMeshProxyComponent()->StaticMesh = AssetInfo.StaticMesh;
-					MeshActor->GetLandscapeMeshProxyComponent()->InitializeForLandscape(AssetInfo.SourceLandscape, SimplificationDetails.LandscapeExportLOD);
+					MeshActor->GetLandscapeMeshProxyComponent()->InitializeForLandscape(AssetInfo.SourceLandscape, AssetInfo.LandscapeLOD);
 					MeshActor->SetActorLabel(AssetInfo.SourceLandscape->GetName());
 				}
 				else
