Fixed: ocassional crash when moving sub-levels in world composition

diff --git a/Engine/Source/Editor/WorldBrowser/Private/Tiles/WorldTileModel.cpp b/Engine/Source/Editor/WorldBrowser/Private/Tiles/WorldTileModel.cpp
--- a/Engine/Source/Editor/WorldBrowser/Private/Tiles/WorldTileModel.cpp
+++ b/Engine/Source/Editor/WorldBrowser/Private/Tiles/WorldTileModel.cpp
@@ -452,7 +452,7 @@ void FWorldTileModel::SetLevelPosition(const FIntPoint& InPosition)
 	
 	// Move actors if necessary
 	ULevel* Level = GetLevelObject();
-	if (Level != NULL && Level->bIsVisible)
+	if (Level != nullptr && Level->bIsVisible)
 	{
 		// Shelve level, if during this translation level will end up out of Editable area
 		if (!ShouldBeVisible(LevelCollectionModel.EditableWorldArea()))
@@ -467,7 +467,10 @@ void FWorldTileModel::SetLevelPosition(const FIntPoint& InPosition)
 
 			for (auto Actor : Level->Actors)
 			{
-				Editor->BroadcastOnActorMoved(Actor);
+				if (Actor != nullptr)
+				{
+					Editor->BroadcastOnActorMoved(Actor);
+				}
 			}
 		}
 	}
