Highlight LOD tiles in different color when previewing streaming in world composition

diff --git a/Engine/Source/Editor/WorldBrowser/Private/Tiles/SWorldTileItem.cpp b/Engine/Source/Editor/WorldBrowser/Private/Tiles/SWorldTileItem.cpp
--- a/Engine/Source/Editor/WorldBrowser/Private/Tiles/SWorldTileItem.cpp
+++ b/Engine/Source/Editor/WorldBrowser/Private/Tiles/SWorldTileItem.cpp
@@ -389,7 +389,8 @@ int32 SWorldTileItem::OnPaint(const FPaintArgs& Args, const FGeometry& AllottedG
 		LayerId = SNodePanel::SNode::OnPaint(Args, AllottedGeometry, ClippingRect, OutDrawElements, LayerId, InWidgetStyle, bParentEnabled);
 		
 		const bool bSelected = (IsItemSelected() || bAffectedByMarquee);
-		const bool bHighlighted = (WorldModel->GetPreviewStreamingLevels().Find(TileModel->TileDetails->PackageName) != nullptr);
+		const int32* PreviewLODIndex = WorldModel->GetPreviewStreamingLevels().Find(TileModel->TileDetails->PackageName);
+		const bool bHighlighted = (PreviewLODIndex != nullptr);
 
 		// Draw the node's selection/highlight.
 		if (bSelected || bHighlighted)
@@ -400,7 +401,13 @@ int32 SWorldTileItem::OnPaint(const FPaintArgs& Args, const FGeometry& AllottedG
 			FSlateLayoutTransform LayoutTransform(Scale, AllottedGeometry.GetAccumulatedLayoutTransform().GetTranslation() - InflateAmount);
 			FSlateRenderTransform RenderTransform(Scale, AllottedGeometry.GetAccumulatedRenderTransform().GetTranslation() - InflateAmount);
 			FPaintGeometry SelectionGeometry(LayoutTransform, RenderTransform, (AllottedGeometry.GetLocalSize()*AllottedGeometry.Scale + InflateAmount*2)/Scale);
-										
+			FLinearColor HighlightColor = FLinearColor::White;
+			if (PreviewLODIndex)
+			{
+				// Highlight LOD tiles in different color to normal tiles
+				HighlightColor = (*PreviewLODIndex == INDEX_NONE) ? FLinearColor::Green : FLinearColor(0.3f,1.0f,0.3f);
+			}
+													
 			FSlateDrawElement::MakeBox(
 				OutDrawElements,
 				LayerId + 1,
@@ -408,7 +415,7 @@ int32 SWorldTileItem::OnPaint(const FPaintArgs& Args, const FGeometry& AllottedG
 				GetShadowBrush(bSelected || bHighlighted),
 				ClippingRect,
 				ESlateDrawEffect::None,
-				bHighlighted ? FLinearColor::Green : FLinearColor::White
+				HighlightColor
 				);
 		}
 
diff --git a/Engine/Source/Editor/WorldBrowser/Private/Tiles/WorldTileCollectionModel.cpp b/Engine/Source/Editor/WorldBrowser/Private/Tiles/WorldTileCollectionModel.cpp
--- a/Engine/Source/Editor/WorldBrowser/Private/Tiles/WorldTileCollectionModel.cpp
+++ b/Engine/Source/Editor/WorldBrowser/Private/Tiles/WorldTileCollectionModel.cpp
@@ -899,7 +899,7 @@ void FWorldTileCollectionModel::DeselectLevels(const FWorldTileLayer& InLayer)
 	}
 }
 
-const TSet<FName>& FWorldTileCollectionModel::GetPreviewStreamingLevels() const
+const TMap<FName, int32>& FWorldTileCollectionModel::GetPreviewStreamingLevels() const
 {
 	return PreviewVisibleTiles;
 }
@@ -922,7 +922,7 @@ void FWorldTileCollectionModel::UpdateStreamingPreview(FVector2D InLocation, boo
 
 			for (const auto& VisibleLevel : DistanceVisibleLevels)
 			{
-				PreviewVisibleTiles.Add(VisibleLevel.StreamingLevel->GetWorldAssetPackageFName());
+				PreviewVisibleTiles.Add(VisibleLevel.StreamingLevel->GetWorldAssetPackageFName(), VisibleLevel.LODIndex);
 			}
 		}
 	}
diff --git a/Engine/Source/Editor/WorldBrowser/Private/Tiles/WorldTileCollectionModel.h b/Engine/Source/Editor/WorldBrowser/Private/Tiles/WorldTileCollectionModel.h
--- a/Engine/Source/Editor/WorldBrowser/Private/Tiles/WorldTileCollectionModel.h
+++ b/Engine/Source/Editor/WorldBrowser/Private/Tiles/WorldTileCollectionModel.h
@@ -114,7 +114,7 @@ public:
 	void UpdateStreamingPreview(FVector2D InPreviewLocation, bool bEnabled);
 
 	/** Returns list of visible streaming levels for current preview location */
-	const TSet<FName>& GetPreviewStreamingLevels() const;
+	const TMap<FName, int32>& GetPreviewStreamingLevels() const;
 		
 	/** Calculates snapped moving delta based on specified landscape tile */
 	FVector2D SnapTranslationDeltaLandscape(const TSharedPtr<FWorldTileModel>& LandscapeTile, 
@@ -252,7 +252,7 @@ private:
 	FLevelModelList						StaticTileList;	
 
 	/** Cached streaming tiles which are potentially visible from specified view point*/
-	TSet<FName>							PreviewVisibleTiles;
+	TMap<FName, int32>					PreviewVisibleTiles;
 	
 	/** View point location for calculating potentially visible streaming tiles*/
 	FVector								PreviewLocation;
