Changed level name tooltip in the Levels Browser to display the long package name of the level (with full path).

diff --git a/Engine/Source/Editor/WorldBrowser/Private/SWorldHierarchyItem.cpp b/Engine/Source/Editor/WorldBrowser/Private/SWorldHierarchyItem.cpp
--- a/Engine/Source/Editor/WorldBrowser/Private/SWorldHierarchyItem.cpp
+++ b/Engine/Source/Editor/WorldBrowser/Private/SWorldHierarchyItem.cpp
@@ -69,7 +69,7 @@ TSharedRef< SWidget > SWorldHierarchyItem::GenerateWidgetForColumn( const FName&
 				.Text(this, &SWorldHierarchyItem::GetLevelDisplayNameText)
 				.ColorAndOpacity(this, &SWorldHierarchyItem::GetLevelDisplayNameColorAndOpacity)
 				.HighlightText(HighlightText)
-				.ToolTipText(LOCTEXT("DoubleClickToolTip", "Double-Click to make this the current Level"))
+				.ToolTipText(this, &SWorldHierarchyItem::GetLevelDisplayNameTooltip)
 			]
 		;
 	}
@@ -203,6 +203,19 @@ FText SWorldHierarchyItem::GetLevelDisplayNameText() const
 	return FText::FromString(LevelModel->GetDisplayName());
 }
 
+FText SWorldHierarchyItem::GetLevelDisplayNameTooltip() const
+{
+	FString PackageName = LevelModel->GetLongPackageName().ToString();
+	if (FPackageName::DoesPackageExist(PackageName))
+	{
+		return FText::FromString(PackageName);
+	}
+	else
+	{
+		return LOCTEXT("UnsavedLevel", "Unsaved Level");
+	}
+}
+
 bool SWorldHierarchyItem::IsSaveEnabled() const
 {
 	return LevelModel->IsLoaded();
diff --git a/Engine/Source/Editor/WorldBrowser/Private/SWorldHierarchyItem.h b/Engine/Source/Editor/WorldBrowser/Private/SWorldHierarchyItem.h
--- a/Engine/Source/Editor/WorldBrowser/Private/SWorldHierarchyItem.h
+++ b/Engine/Source/Editor/WorldBrowser/Private/SWorldHierarchyItem.h
@@ -105,7 +105,12 @@ private:
 	 *  @return The text of level name while it is not being edited
 	 */
 	FText GetLevelDisplayNameText() const;
-	
+
+	/**
+	*  @return The tooltip of level name (shows the full path of the asset package)
+	*/
+	FText GetLevelDisplayNameTooltip() const;
+
 	/** */
 	FSlateFontInfo GetLevelDisplayNameFont() const;
 	
