Fixed the text for remove level.

UE-14584

diff --git a/Engine/Source/Editor/WorldBrowser/Private/LevelCollectionCommands.h b/Engine/Source/Editor/WorldBrowser/Private/LevelCollectionCommands.h
--- a/Engine/Source/Editor/WorldBrowser/Private/LevelCollectionCommands.h
+++ b/Engine/Source/Editor/WorldBrowser/Private/LevelCollectionCommands.h
@@ -42,7 +42,7 @@ public:
 		UI_COMMAND( World_CreateEmptyLevel, "Create New...", "Creates a new empty Level", EUserInterfaceActionType::Button, FInputChord() );
 		UI_COMMAND( World_AddExistingLevel, "Add Existing...", "Adds an existing level", EUserInterfaceActionType::Button, FInputChord() );
 		UI_COMMAND( World_AddSelectedActorsToNewLevel, "Create New with Selected Actors...", "Adds the actors currently selected in the active viewport to a new Level", EUserInterfaceActionType::Button, FInputChord() );
-		UI_COMMAND( World_RemoveSelectedLevels, "Remove Selected", "Removes selected levels from the disk", EUserInterfaceActionType::Button, FInputChord() );
+		UI_COMMAND(World_RemoveSelectedLevels, "Remove Selected", "Removes selected levels from the base streaming level", EUserInterfaceActionType::Button, FInputChord() );
 
 		UI_COMMAND( MoveWorldOrigin, "Move World Origin to Level Position", "Moves world origin to level position", EUserInterfaceActionType::Button, FInputChord() );
 		UI_COMMAND( ResetWorldOrigin, "Reset World Origin", "Moves world origin to zero", EUserInterfaceActionType::Button, FInputChord() );
