Restricting number of world composition tile thumbnail updates per frame
UE-16818

diff --git a/Engine/Source/Editor/WorldBrowser/Private/Tiles/SWorldTileItem.cpp b/Engine/Source/Editor/WorldBrowser/Private/Tiles/SWorldTileItem.cpp
--- a/Engine/Source/Editor/WorldBrowser/Private/Tiles/SWorldTileItem.cpp
+++ b/Engine/Source/Editor/WorldBrowser/Private/Tiles/SWorldTileItem.cpp
@@ -380,10 +380,11 @@ int32 SWorldTileItem::OnPaint(const FPaintArgs& Args, const FGeometry& AllottedG
 	if (bIsVisible)
 	{
 		// Redraw thumbnail image if requested
-		if (bNeedRefresh)
+		if (bNeedRefresh && WorldModel->ShouldUpdateTileThumbnail(TileModel))
 		{
-			Thumbnail->UpdateThumbnail();
 			bNeedRefresh = false;
+			Thumbnail->UpdateThumbnail();
+			WorldModel->OnTileThumbnailUpdated(TileModel);
 		}
 		
 		LayerId = SNodePanel::SNode::OnPaint(Args, AllottedGeometry, ClippingRect, OutDrawElements, LayerId, InWidgetStyle, bParentEnabled);
diff --git a/Engine/Source/Editor/WorldBrowser/Private/Tiles/WorldTileCollectionModel.cpp b/Engine/Source/Editor/WorldBrowser/Private/Tiles/WorldTileCollectionModel.cpp
--- a/Engine/Source/Editor/WorldBrowser/Private/Tiles/WorldTileCollectionModel.cpp
+++ b/Engine/Source/Editor/WorldBrowser/Private/Tiles/WorldTileCollectionModel.cpp
@@ -34,6 +34,8 @@ static const FName HeightmapLayerName = FName("__Heightmap__");
 
 FWorldTileCollectionModel::FWorldTileCollectionModel(UEditorEngine* InEditor)
 	: FLevelCollectionModel(InEditor)
+	, PreviewLocation(0.0f,0.0f,0.0f)
+	, TileThumbnailUpdateCooldown(0.0f)
 	, bIsSavingLevel(false)
 	, bMeshProxyAvailable(false)
 {
@@ -83,6 +85,8 @@ void FWorldTileCollectionModel::Tick( float DeltaTime )
 		return;
 	}
 
+	TileThumbnailUpdateCooldown = FMath::Max(0.0f, TileThumbnailUpdateCooldown - DeltaTime);
+
 	FLevelCollectionModel::Tick(DeltaTime);
 }
 
@@ -621,6 +625,16 @@ bool FWorldTileCollectionModel::GetObserverView(FVector& Location, FRotator& Rot
 	return false;
 }
 
+bool FWorldTileCollectionModel::ShouldUpdateTileThumbnail(const TSharedPtr<FWorldTileModel>& TileModel)
+{
+	return (TileThumbnailUpdateCooldown <= 0.f);
+}
+	
+void FWorldTileCollectionModel::OnTileThumbnailUpdated(const TSharedPtr<FWorldTileModel>& TileModel)
+{
+	TileThumbnailUpdateCooldown = 0.1f;
+}
+
 static double Area(FVector2D InRect)
 {
 	return (double)InRect.X * (double)InRect.Y;
diff --git a/Engine/Source/Editor/WorldBrowser/Private/Tiles/WorldTileCollectionModel.h b/Engine/Source/Editor/WorldBrowser/Private/Tiles/WorldTileCollectionModel.h
--- a/Engine/Source/Editor/WorldBrowser/Private/Tiles/WorldTileCollectionModel.h
+++ b/Engine/Source/Editor/WorldBrowser/Private/Tiles/WorldTileCollectionModel.h
@@ -143,6 +143,12 @@ public:
 	 *  Builds context menu for a world composition
 	 */
 	void BuildWorldCompositionMenu(FMenuBuilder& InMenuBuilder) const;
+	
+	/** Whether tile view should update tile thumbnail */
+	bool ShouldUpdateTileThumbnail(const TSharedPtr<FWorldTileModel>& TileModel);
+	
+	/** Notification from tile view that tile thumbnail was updated */
+	void OnTileThumbnailUpdated(const TSharedPtr<FWorldTileModel>& TileModel);
 
 private:
 	FWorldTileCollectionModel(UEditorEngine* InEditor);
@@ -266,6 +272,8 @@ private:
 	/** All selected layers */
 	TArray<FWorldTileLayer>				SelectedLayers;
 
+	float								TileThumbnailUpdateCooldown;
+
 	// Is in process of saving a level
 	bool								bIsSavingLevel;
 	
