Fixed case where world composition view jumps to wrong location when selected tile is hidden

diff --git a/Engine/Source/Editor/WorldBrowser/Private/Tiles/WorldTileModel.cpp b/Engine/Source/Editor/WorldBrowser/Private/Tiles/WorldTileModel.cpp
--- a/Engine/Source/Editor/WorldBrowser/Private/Tiles/WorldTileModel.cpp
+++ b/Engine/Source/Editor/WorldBrowser/Private/Tiles/WorldTileModel.cpp
@@ -154,7 +154,7 @@ void FWorldTileModel::UpdateAsset(const FAssetData& AssetData)
 
 FVector2D FWorldTileModel::GetLevelPosition2D() const
 {
-	if (TileDetails->Bounds.IsValid)
+	if (TileDetails->Bounds.IsValid && !TileDetails->bHideInTileView)
 	{
 		FVector2D LevelPosition = GetLevelCurrentPosition();
 		return LevelPosition - FVector2D(TileDetails->Bounds.GetExtent()) + GetLevelTranslationDelta();
@@ -165,7 +165,7 @@ FVector2D FWorldTileModel::GetLevelPosition2D() const
 
 FVector2D FWorldTileModel::GetLevelSize2D() const
 {
-	if (TileDetails->Bounds.IsValid)
+	if (TileDetails->Bounds.IsValid && !TileDetails->bHideInTileView)
 	{
 		FVector LevelSize = TileDetails->Bounds.GetSize();
 		return FVector2D(LevelSize.X, LevelSize.Y);
