World composition packs tile images into 4x4 atlases in order to speed up minimap drawing when there are hundreds of tiles
UE-16821

diff --git a/Engine/Source/Editor/WorldBrowser/Private/Tiles/SWorldComposition.cpp b/Engine/Source/Editor/WorldBrowser/Private/Tiles/SWorldComposition.cpp
--- a/Engine/Source/Editor/WorldBrowser/Private/Tiles/SWorldComposition.cpp
+++ b/Engine/Source/Editor/WorldBrowser/Private/Tiles/SWorldComposition.cpp
@@ -10,6 +10,7 @@
 #include "SWorldLayers.h"
 #include "WorldTileCollectionModel.h"
 #include "IMenu.h"
+#include "WorldTileThumbnails.h"
 
 #define LOCTEXT_NAMESPACE "WorldBrowser"
 
@@ -79,22 +80,10 @@ public:
 		, bHasNodeInteraction(true)
 		, BoundsSnappingDistance(20.f)
 	{
-
-		SharedThumbnailRT = new FSlateTextureRenderTarget2DResource(
-						FLinearColor::Black, 
-						512, 
-						512, 
-						PF_B8G8R8A8, SF_Point, TA_Wrap, TA_Wrap, 0.0f
-					);
-		BeginInitResource(SharedThumbnailRT);
 	}
 
 	~SWorldCompositionGrid()
 	{
-		BeginReleaseResource(SharedThumbnailRT);
-		FlushRenderingCommands();
-		delete SharedThumbnailRT;
-		
 		WorldModel->SelectionChanged.RemoveAll(this);
 		WorldModel->CollectionChanged.RemoveAll(this);
 
@@ -140,6 +129,8 @@ public:
 		SelectionManager.OnSelectionChanged.BindSP(this, &SWorldCompositionGrid::OnSelectionChanged);
 
 		FCoreDelegates::PreWorldOriginOffset.AddSP(this, &SWorldCompositionGrid::PreWorldOriginOffset);
+
+		ThumbnailCollection = MakeShareable(new FTileThumbnailCollection());
 	
 		RefreshView();
 	}
@@ -150,7 +141,7 @@ public:
 		auto NewNode = SNew(SWorldTileItem)
 							.InWorldModel(WorldModel)
 							.InItemModel(LevelModel)
-							.ThumbnailRenderTarget(SharedThumbnailRT);
+							.InThumbnailCollection(ThumbnailCollection);
 	
 		AddGraphNode(NewNode);
 	}
@@ -949,8 +940,8 @@ private:
 	FVector2D								WorldMouseLocation;
 	// Current marquee rectangle size in world units
 	FVector2D								WorldMarqueeSize;
-
-	FSlateTextureRenderTarget2DResource*	SharedThumbnailRT;
+	// Thumbnail managment for tile items
+	TSharedPtr<FTileThumbnailCollection>	ThumbnailCollection;
 };
 
 //----------------------------------------------------------------
diff --git a/Engine/Source/Editor/WorldBrowser/Private/Tiles/SWorldTileItem.cpp b/Engine/Source/Editor/WorldBrowser/Private/Tiles/SWorldTileItem.cpp
--- a/Engine/Source/Editor/WorldBrowser/Private/Tiles/SWorldTileItem.cpp
+++ b/Engine/Source/Editor/WorldBrowser/Private/Tiles/SWorldTileItem.cpp
@@ -1,134 +1,53 @@
 // Copyright 1998-2015 Epic Games, Inc. All Rights Reserved.
 #include "WorldBrowserPrivatePCH.h"
 
-#include "ObjectTools.h"
 #include "EdGraphUtilities.h"
 #include "WorldTileCollectionModel.h"
 #include "SWorldTileItem.h"
 #include "WorldTileDetails.h"
+#include "WorldTileThumbnails.h"
 
 #define LOCTEXT_NAMESPACE "WorldBrowser"
 
-FTileItemThumbnail::FTileItemThumbnail(FSlateTextureRenderTarget2DResource* InThumbnailRenderTarget, 
-									   TSharedPtr<FLevelModel> InItemModel)
-	: LevelModel(InItemModel)
-	, ThumbnailTexture(NULL)
-	, ThumbnailRenderTarget(InThumbnailRenderTarget)
-{
-	FIntPoint RTSize = ThumbnailRenderTarget->GetSizeXY();
-	ThumbnailTexture = new FSlateTexture2DRHIRef(RTSize.X, RTSize.Y, PF_B8G8R8A8, NULL, TexCreate_Dynamic);
-	BeginInitResource(ThumbnailTexture);
-}
-
-FTileItemThumbnail::~FTileItemThumbnail()
-{
-	BeginReleaseResource(ThumbnailTexture);
-	
-	// Wait for all resources to be released
-	FlushRenderingCommands();
-	delete ThumbnailTexture;
-	ThumbnailTexture = NULL;
-}
-
-FIntPoint FTileItemThumbnail::GetSize() const
-{
-	return ThumbnailRenderTarget->GetSizeXY();
-}
-
-FSlateShaderResource* FTileItemThumbnail::GetViewportRenderTargetTexture() const
-{
-	return ThumbnailTexture;
-}
-
-bool FTileItemThumbnail::RequiresVsync() const
-{
-	return false;
-}
-
-void FTileItemThumbnail::UpdateTextureData(FObjectThumbnail* ObjectThumbnail)
+//----------------------------------------------------------------
+//
+//
+//----------------------------------------------------------------
+int32 SWorldTileImage::OnPaint(const FPaintArgs& Args, const FGeometry& AllottedGeometry, const FSlateRect& MyClippingRect, FSlateWindowElementList& OutDrawElements, int32 LayerId, const FWidgetStyle& InWidgetStyle, bool bParentEnabled) const
 {
-	check(ThumbnailTexture)
-	
-	if (ObjectThumbnail &&
-		ObjectThumbnail->GetImageWidth() > 0 && 
-		ObjectThumbnail->GetImageHeight() > 0 && 
-		ObjectThumbnail->GetUncompressedImageData().Num() > 0)
-	{
-		// Make bulk data for updating the texture memory later
-		FSlateTextureDataPtr ThumbnailBulkData = MakeShareable(new FSlateTextureData(
-			ObjectThumbnail->GetImageWidth(), 
-			ObjectThumbnail->GetImageHeight(), 
-			GPixelFormats[PF_B8G8R8A8].BlockBytes, 
-			ObjectThumbnail->AccessImageData())
-			);
-					
-		// Update the texture RHI
-		ENQUEUE_UNIQUE_RENDER_COMMAND_TWOPARAMETER(
-			UpdateGridItemThumbnailResourse,
-			FSlateTexture2DRHIRef*, Texture, ThumbnailTexture,
-			FSlateTextureDataPtr, BulkData, ThumbnailBulkData,
-		{
-			Texture->SetTextureData(BulkData);
-			Texture->UpdateRHI();
-		});
-	}
-}
+	const FSlateBrush* ImageBrush = Image.Get();
 
-void FTileItemThumbnail::UpdateThumbnail()
-{
-	// No need images for persistent and always loaded levels
-	if (LevelModel->IsPersistent())
-	{
-		return;
-	}
-	
-	// Load image from a package header
-	if (!LevelModel->IsVisible() || LevelModel->IsSimulating())
-	{
-		TSet<FName> ObjectFullNames;
-		ObjectFullNames.Add(LevelModel->GetAssetName());
-		FThumbnailMap ThumbnailMap;
-			
-		if (ThumbnailTools::ConditionallyLoadThumbnailsFromPackage(LevelModel->GetPackageFileName(), ObjectFullNames, ThumbnailMap))
-		{
-			UpdateTextureData(ThumbnailMap.Find(LevelModel->GetAssetName()));
-		}
-	}
-	// Render image from a loaded visble level
-	else
+	if ((ImageBrush != nullptr) && (ImageBrush->DrawAs != ESlateBrushDrawType::NoDrawType))
 	{
-		ULevel* TargetLevel = LevelModel->GetLevelObject();
-		if (TargetLevel)
-		{
-			FIntPoint RTSize = ThumbnailRenderTarget->GetSizeXY();
-			
-			// Set persistent world package as transient to avoid package dirtying during thumbnail rendering
-			FUnmodifiableObject ImmuneWorld(TargetLevel->OwningWorld);
-			
-			FObjectThumbnail NewThumbnail;
-			// Generate the thumbnail
-			ThumbnailTools::RenderThumbnail(
-				TargetLevel,
-				RTSize.X,
-				RTSize.Y,
-				ThumbnailTools::EThumbnailTextureFlushMode::NeverFlush,
-				ThumbnailRenderTarget,
-				&NewThumbnail
-				);
-
-			UPackage* MyOutermostPackage = CastChecked<UPackage>(TargetLevel->GetOutermost());
-			ThumbnailTools::CacheThumbnail(LevelModel->GetAssetName().ToString(), &NewThumbnail, MyOutermostPackage);
-			UpdateTextureData(&NewThumbnail);
-		}
+		const bool bIsEnabled = EditableTile.Get() && ShouldBeEnabled(bParentEnabled);
+		const uint32 DrawEffects = bIsEnabled ? 0 : ESlateDrawEffect::DisabledEffect;
+		FSlateDrawElement::MakeBox(
+			OutDrawElements, 
+			LayerId, 
+			AllottedGeometry.ToPaintGeometry(), 
+			ImageBrush, 
+			MyClippingRect, 
+			DrawEffects|ESlateDrawEffect::IgnoreTextureAlpha, 
+			FColor::White);
 	}
+	return LayerId;
 }
 
 //----------------------------------------------------------------
 //
 //
 //----------------------------------------------------------------
+SWorldTileItem::SWorldTileItem()
+ : bAffectedByMarquee(false)
+ , bNeedRefresh(false)
+ , bIsDragging(false)
+{
+}
+
 SWorldTileItem::~SWorldTileItem()
 {
+	ThumbnailImageWidget->SetImage(nullptr);
+	ThumbnailCollection->UnregisterTile(*TileModel.Get());
 	TileModel->ChangedEvent.RemoveAll(this);
 }
 
@@ -136,35 +55,23 @@ void SWorldTileItem::Construct(const FArguments& InArgs)
 {
 	WorldModel = InArgs._InWorldModel;
 	TileModel = InArgs._InItemModel;
-
-	ProgressBarImage = FEditorStyle::GetBrush(TEXT("ProgressBar.Marquee"));
+	ThumbnailCollection = InArgs._InThumbnailCollection;
 	
-	bNeedRefresh = true;
-	bIsDragging = false;
-	bAffectedByMarquee = false;
-
-	Thumbnail = MakeShareable(new FTileItemThumbnail(InArgs._ThumbnailRenderTarget, TileModel));
-	ThumbnailViewport = SNew(SViewport)
-							.EnableGammaCorrection(false);
-							
-	ThumbnailViewport->SetViewportInterface(Thumbnail.ToSharedRef());
-	// This will grey out tile in case level is not loaded or locked
-	ThumbnailViewport->SetEnabled(TAttribute<bool>(this, &SWorldTileItem::IsItemEnabled));
-
 	TileModel->ChangedEvent.AddSP(this, &SWorldTileItem::RequestRefresh);
-			
+	
 	GetOrAddSlot( ENodeZone::Center )
 	[
-		ThumbnailViewport.ToSharedRef()
+		SAssignNew(ThumbnailImageWidget, SWorldTileImage)
+		.EditableTile(this, &SWorldTileItem::IsItemEnabled)
 	];
 
+	ThumbnailCollection->RegisterTile(*TileModel.Get());
+	const FSlateBrush* TileBrush = ThumbnailCollection->GetTileBrush(*TileModel.Get());
+	ThumbnailImageWidget->SetImage(TileBrush);
+
 	SetToolTip(CreateToolTipWidget());
 
-	CurveSequence = FCurveSequence(0.0f, 0.5f);
-	if (TileModel->IsLoading())
-	{
-		CurveSequence.Play( this->AsShared(), true );
-	}
+	bNeedRefresh = true;
 }
 
 void SWorldTileItem::RequestRefresh()
@@ -361,18 +268,6 @@ FVector2D SWorldTileItem::ComputeDesiredSize( float ) const
 	return TileModel->GetLevelSize2D();
 }
 
-void SWorldTileItem::Tick( const FGeometry& AllottedGeometry, const double InCurrentTime, const float InDeltaTime )
-{
-	if ( TileModel->IsLoading() && !CurveSequence.IsPlaying() )
-	{
-		CurveSequence.Play( this->AsShared() );
-	}
-	else if ( !TileModel->IsLoading() && CurveSequence.IsPlaying() )
-	{
-		CurveSequence.JumpToEnd();
-	}
-}
-
 int32 SWorldTileItem::OnPaint(const FPaintArgs& Args, const FGeometry& AllottedGeometry, const FSlateRect& ClippingRect, FSlateWindowElementList& OutDrawElements, int32 LayerId, const FWidgetStyle& InWidgetStyle, bool bParentEnabled) const
 {
 	const bool bIsVisible = FSlateRect::DoRectanglesIntersect(AllottedGeometry.GetClippingRect(), ClippingRect);
@@ -380,11 +275,11 @@ int32 SWorldTileItem::OnPaint(const FPaintArgs& Args, const FGeometry& AllottedG
 	if (bIsVisible)
 	{
 		// Redraw thumbnail image if requested
-		if (bNeedRefresh && WorldModel->ShouldUpdateTileThumbnail(TileModel))
+		if (bNeedRefresh && !ThumbnailCollection->IsOnCooldown())
 		{
 			bNeedRefresh = false;
-			Thumbnail->UpdateThumbnail();
-			WorldModel->OnTileThumbnailUpdated(TileModel);
+			const FSlateBrush* TileBrush = ThumbnailCollection->UpdateTileThumbnail(*TileModel.Get());
+			ThumbnailImageWidget->SetImage(TileBrush);
 		}
 		
 		LayerId = SNodePanel::SNode::OnPaint(Args, AllottedGeometry, ClippingRect, OutDrawElements, LayerId, InWidgetStyle, bParentEnabled);
@@ -419,34 +314,6 @@ int32 SWorldTileItem::OnPaint(const FPaintArgs& Args, const FGeometry& AllottedG
 				HighlightColor
 				);
 		}
-
-		// Draw progress bar if level is currently loading
-		if (TileModel->IsLoading())
-		{
-			const float ProgressBarAnimOffset = ProgressBarImage->ImageSize.X * CurveSequence.GetLerp();
-			const float ProgressBarImageSize = ProgressBarImage->ImageSize.X;
-			const float ProgressBarImageHeight = ProgressBarImage->ImageSize.Y;
-			const FVector2D ProggresBarOffset = FVector2D(ProgressBarAnimOffset - ProgressBarImageSize, 0);
-
-			FSlateLayoutTransform ProgressBarLayoutTransform(1.f, AllottedGeometry.GetAccumulatedLayoutTransform().GetTranslation() + ProggresBarOffset);
-			FSlateRenderTransform ProgressBarRenderTransform(1.f, AllottedGeometry.GetAccumulatedRenderTransform().GetTranslation() + ProggresBarOffset);
-			FPaintGeometry ProgressBarPaintGeometry(
-				ProgressBarLayoutTransform, 
-				ProgressBarRenderTransform, 
-				FVector2D(AllottedGeometry.Size.X*AllottedGeometry.Scale + ProgressBarImageSize, FMath::Min(AllottedGeometry.Size.Y*AllottedGeometry.Scale, ProgressBarImageHeight)));
-								
-			const FSlateRect ProgressBarClippingRect = AllottedGeometry.GetClippingRect().IntersectionWith(ClippingRect);
-								
-			FSlateDrawElement::MakeBox(
-				OutDrawElements,
-				LayerId + 1,
-				ProgressBarPaintGeometry,
-				ProgressBarImage,
-				ProgressBarClippingRect,
-				ESlateDrawEffect::None,
-				FLinearColor(1.0f, 1.0f, 1.0f, 0.5f)
-			);
-		}
 	}
 	
 	return LayerId;
diff --git a/Engine/Source/Editor/WorldBrowser/Private/Tiles/SWorldTileItem.h b/Engine/Source/Editor/WorldBrowser/Private/Tiles/SWorldTileItem.h
--- a/Engine/Source/Editor/WorldBrowser/Private/Tiles/SWorldTileItem.h
+++ b/Engine/Source/Editor/WorldBrowser/Private/Tiles/SWorldTileItem.h
@@ -1,34 +1,33 @@
 // Copyright 1998-2015 Epic Games, Inc. All Rights Reserved.
 #pragma once
 
-#include "SlateTextures.h"
 #include "GraphEditor.h"
 #include "SNodePanel.h"
-#include "AssetThumbnail.h"
 
-class FTileItemThumbnail : public ISlateViewport, public TSharedFromThis<FTileItemThumbnail>
+class FTileThumbnailCollection;
+
+//----------------------------------------------------------------
+//
+//
+//----------------------------------------------------------------
+class SWorldTileImage : public SImage
 {
-public:
-	FTileItemThumbnail(FSlateTextureRenderTarget2DResource* InThumbnailRenderTarget, TSharedPtr<FLevelModel> InItemModel);
-	~FTileItemThumbnail();
+	SLATE_BEGIN_ARGS(SWorldTileImage)
+		: _EditableTile(false)
+	{}
+		SLATE_ATTRIBUTE(bool, EditableTile)
+	SLATE_END_ARGS()
 
-	/* ISlateViewport interface */
-	virtual FIntPoint GetSize() const override;
-	virtual class FSlateShaderResource* GetViewportRenderTargetTexture() const override;
-	virtual bool RequiresVsync() const override;
-	
-	/** Request thumbnail redraw*/
-	void UpdateThumbnail();
+public:
+	void Construct(const FArguments& InArgs)
+	{
+		EditableTile = InArgs._EditableTile;
+	}
 
-private:
-	void UpdateTextureData(FObjectThumbnail* ObjectThumbnail);
+	virtual int32 OnPaint(const FPaintArgs& Args, const FGeometry& AllottedGeometry, const FSlateRect& MyClippingRect, FSlateWindowElementList& OutDrawElements, int32 LayerId, const FWidgetStyle& InWidgetStyle, bool bParentEnabled) const override;
 
 private:
-	TSharedPtr<FLevelModel>					LevelModel;
-	/** Rendering resource for slate */
-	FSlateTexture2DRHIRef*					ThumbnailTexture;
-	/** Shared render target for slate */
-	FSlateTextureRenderTarget2DResource*	ThumbnailRenderTarget;
+	TAttribute<bool> EditableTile;
 };
 
 //----------------------------------------------------------------
@@ -45,15 +44,15 @@ public:
 		SLATE_ARGUMENT(TSharedPtr<FWorldTileCollectionModel>, InWorldModel)
 		/** Data for the asset this item represents */
 		SLATE_ARGUMENT(TSharedPtr<FWorldTileModel>, InItemModel)
-		//
-		SLATE_ARGUMENT(FSlateTextureRenderTarget2DResource*, ThumbnailRenderTarget)
+		/** Thumbnails management */
+		SLATE_ARGUMENT(TSharedPtr<FTileThumbnailCollection>, InThumbnailCollection)
 	SLATE_END_ARGS()
 
+	SWorldTileItem();
 	~SWorldTileItem();
 
 	void Construct(const FArguments& InArgs);
-	virtual void Tick( const FGeometry& AllottedGeometry, const double InCurrentTime, const float InDeltaTime ) override;
-	
+		
 	// SNodePanel::SNode interface start
 	virtual FVector2D GetDesiredSizeForMarquee() const override;
 	virtual UObject* GetObjectBeingDisplayed() const override;
@@ -103,26 +102,18 @@ private:
 	FText GetLevelLayerDistanceText() const;
 	
 public:
-	bool							bAffectedByMarquee;
+	bool									bAffectedByMarquee;
 
 private:
 	/** The world data */
-	TSharedPtr<FWorldTileCollectionModel> WorldModel;
-
+	TSharedPtr<FWorldTileCollectionModel>	WorldModel;
 	/** The data for this item */
-	TSharedPtr<FWorldTileModel>		TileModel;
-
-	mutable bool					bNeedRefresh;
-	bool							bIsDragging;
-
-	TSharedPtr<FTileItemThumbnail>	Thumbnail;
-	/** The actual widget for the thumbnail. */
-	TSharedPtr<SViewport>			ThumbnailViewport;
+	TSharedPtr<FWorldTileModel>				TileModel;
 	
-	/** */
-	const FSlateBrush*				ProgressBarImage;
-
-	/** Curve sequence to drive loading animation */
-	FCurveSequence					CurveSequence;
+	TSharedPtr<SWorldTileImage>				ThumbnailImageWidget;	
+	TSharedPtr<FTileThumbnailCollection>	ThumbnailCollection;
+	
+	mutable bool							bNeedRefresh;
+	bool									bIsDragging;
 };
 
diff --git a/Engine/Source/Editor/WorldBrowser/Private/Tiles/WorldTileCollectionModel.cpp b/Engine/Source/Editor/WorldBrowser/Private/Tiles/WorldTileCollectionModel.cpp
--- a/Engine/Source/Editor/WorldBrowser/Private/Tiles/WorldTileCollectionModel.cpp
+++ b/Engine/Source/Editor/WorldBrowser/Private/Tiles/WorldTileCollectionModel.cpp
@@ -35,7 +35,6 @@ static const FName HeightmapLayerName = FName("__Heightmap__");
 FWorldTileCollectionModel::FWorldTileCollectionModel(UEditorEngine* InEditor)
 	: FLevelCollectionModel(InEditor)
 	, PreviewLocation(0.0f,0.0f,0.0f)
-	, TileThumbnailUpdateCooldown(0.0f)
 	, bIsSavingLevel(false)
 	, bMeshProxyAvailable(false)
 {
@@ -85,8 +84,6 @@ void FWorldTileCollectionModel::Tick( float DeltaTime )
 		return;
 	}
 
-	TileThumbnailUpdateCooldown = FMath::Max(0.0f, TileThumbnailUpdateCooldown - DeltaTime);
-
 	FLevelCollectionModel::Tick(DeltaTime);
 }
 
@@ -625,16 +622,6 @@ bool FWorldTileCollectionModel::GetObserverView(FVector& Location, FRotator& Rot
 	return false;
 }
 
-bool FWorldTileCollectionModel::ShouldUpdateTileThumbnail(const TSharedPtr<FWorldTileModel>& TileModel)
-{
-	return (TileThumbnailUpdateCooldown <= 0.f);
-}
-	
-void FWorldTileCollectionModel::OnTileThumbnailUpdated(const TSharedPtr<FWorldTileModel>& TileModel)
-{
-	TileThumbnailUpdateCooldown = 0.1f;
-}
-
 static double Area(FVector2D InRect)
 {
 	return (double)InRect.X * (double)InRect.Y;
diff --git a/Engine/Source/Editor/WorldBrowser/Private/Tiles/WorldTileCollectionModel.h b/Engine/Source/Editor/WorldBrowser/Private/Tiles/WorldTileCollectionModel.h
--- a/Engine/Source/Editor/WorldBrowser/Private/Tiles/WorldTileCollectionModel.h
+++ b/Engine/Source/Editor/WorldBrowser/Private/Tiles/WorldTileCollectionModel.h
@@ -143,12 +143,6 @@ public:
 	 *  Builds context menu for a world composition
 	 */
 	void BuildWorldCompositionMenu(FMenuBuilder& InMenuBuilder) const;
-	
-	/** Whether tile view should update tile thumbnail */
-	bool ShouldUpdateTileThumbnail(const TSharedPtr<FWorldTileModel>& TileModel);
-	
-	/** Notification from tile view that tile thumbnail was updated */
-	void OnTileThumbnailUpdated(const TSharedPtr<FWorldTileModel>& TileModel);
 
 private:
 	FWorldTileCollectionModel(UEditorEngine* InEditor);
@@ -272,8 +266,6 @@ private:
 	/** All selected layers */
 	TArray<FWorldTileLayer>				SelectedLayers;
 
-	float								TileThumbnailUpdateCooldown;
-
 	// Is in process of saving a level
 	bool								bIsSavingLevel;
 	
diff --git a/Engine/Source/Editor/WorldBrowser/Private/Tiles/WorldTileThumbnails.cpp b/Engine/Source/Editor/WorldBrowser/Private/Tiles/WorldTileThumbnails.cpp
new file mode 100644
--- /dev/null
+++ b/Engine/Source/Editor/WorldBrowser/Private/Tiles/WorldTileThumbnails.cpp
@@ -0,0 +1,353 @@
+// Copyright 1998-2015 Epic Games, Inc. All Rights Reserved.
+
+#include "WorldBrowserPrivatePCH.h"
+#include "WorldTileThumbnails.h"
+#include "ObjectTools.h"
+#include "AssetThumbnail.h"
+#include "SlateTextures.h"
+#include "WorldTileCollectionModel.h"
+
+static const double TileThumbnailUpdateCooldown = 0.005f;
+static const FSlateColorBrush ThumbnailDefaultBrush(FLinearColor::White);
+
+//----------------------------------------------------------------
+//
+//
+//----------------------------------------------------------------
+FTileThumbnail::FTileThumbnail(FSlateTextureRenderTarget2DResource* InThumbnailRenderTarget, const FWorldTileModel& InTileModel, const FIntPoint& InSlotAllocation)
+	: TileModel(InTileModel)
+	, ThumbnailRenderTarget(InThumbnailRenderTarget)
+	, SlotAllocation(InSlotAllocation)
+{
+}
+
+FIntPoint FTileThumbnail::GetThumbnailSlotAllocation() const
+{
+	return SlotAllocation;
+}
+
+FSlateTextureDataPtr FTileThumbnail::UpdateThumbnail()
+{
+	// No need images for persistent and always loaded levels
+	if (TileModel.IsPersistent())
+	{
+		return ToSlateTextureData(nullptr);
+	}
+	
+	// Load image from a package header
+	if (!TileModel.IsVisible() || TileModel.IsSimulating())
+	{
+		const FName LevelAssetName = TileModel.GetAssetName();
+		TSet<FName> ObjectFullNames;
+		ObjectFullNames.Add(LevelAssetName);
+		FThumbnailMap ThumbnailMap;
+
+		if (ThumbnailTools::ConditionallyLoadThumbnailsFromPackage(TileModel.GetPackageFileName(), ObjectFullNames, ThumbnailMap))
+		{
+			const FObjectThumbnail* ObjectThumbnail = ThumbnailMap.Find(LevelAssetName);
+			return ToSlateTextureData(ObjectThumbnail);
+		}
+	}
+	// Render image from a visible level
+	else
+	{
+		ULevel* TargetLevel = TileModel.GetLevelObject();
+		if (TargetLevel)
+		{
+			FIntPoint RTSize = ThumbnailRenderTarget->GetSizeXY();
+			
+			// Set persistent world package as transient to avoid package dirtying during thumbnail rendering
+			FUnmodifiableObject ImmuneWorld(TargetLevel->OwningWorld);
+			
+			FObjectThumbnail NewThumbnail;
+			// Generate the thumbnail
+			ThumbnailTools::RenderThumbnail(
+				TargetLevel,
+				RTSize.X,
+				RTSize.Y,
+				ThumbnailTools::EThumbnailTextureFlushMode::NeverFlush,
+				ThumbnailRenderTarget,
+				&NewThumbnail
+				);
+
+			UPackage* MyOutermostPackage = CastChecked<UPackage>(TargetLevel->GetOutermost());
+			ThumbnailTools::CacheThumbnail(TileModel.GetAssetName().ToString(), &NewThumbnail, MyOutermostPackage);
+			return ToSlateTextureData(&NewThumbnail);
+		}
+	}
+
+	return ToSlateTextureData(nullptr);
+}
+
+FSlateTextureDataPtr FTileThumbnail::ToSlateTextureData(const FObjectThumbnail* ObjectThumbnail) const
+{
+	FSlateTextureDataPtr Result;
+	
+	if (ObjectThumbnail)
+	{
+		FIntPoint ImageSize(ObjectThumbnail->GetImageWidth(), ObjectThumbnail->GetImageHeight());
+		FIntPoint TargetSize = ThumbnailRenderTarget->GetSizeXY();
+		if (TargetSize == ImageSize)
+		{
+			const TArray<uint8>& ImageData = ObjectThumbnail->GetUncompressedImageData();
+			if (ImageData.Num())
+			{
+				Result = MakeShareable(new FSlateTextureData(ImageData.GetData(), ImageSize.X, ImageSize.Y, 4));
+			}
+		}
+	}
+
+	return Result;
+}
+
+//----------------------------------------------------------------
+//
+//
+//----------------------------------------------------------------
+
+FTileAtlasPage::FTileAtlasPage()
+{
+	AtlasTexture = UTexture2DDynamic::Create(TileThumbnailAtlasSize, TileThumbnailAtlasSize, PF_B8G8R8A8, false);
+	AtlasTexture->AddToRoot();
+
+	static int32 NextPageUniqueID = 1;
+	FName AtlasPageName = *FString::Printf(TEXT("WorldCompositionAtlasPage_%d"), NextPageUniqueID++);
+	
+	for (int32 i = 0; i < ARRAY_COUNT(AtlasSlots); ++i)
+	{
+		FTileAtlasSlot& Slot = AtlasSlots[i];
+		
+		Slot.bOccupied = false;
+		Slot.SlotBrush = new FSlateDynamicImageBrush(
+			AtlasTexture, 
+			FVector2D(TileThumbnailAtlasSize, TileThumbnailAtlasSize),
+			AtlasPageName
+			);
+
+		int32 SlotOffsetX = i % TileThumbnailAtlasDim;
+		int32 SlotOffsetY = i / TileThumbnailAtlasDim;
+		FVector2D StartUV = FVector2D(SlotOffsetX, SlotOffsetY)/TileThumbnailAtlasDim;
+		FVector2D SizeUV = FVector2D(1.0f, 1.0f)/TileThumbnailAtlasDim;
+
+		Slot.SlotBrush->SetUVRegion(FBox2D(StartUV, StartUV+SizeUV));
+	}
+}
+
+FTileAtlasPage::~FTileAtlasPage()
+{
+	for (int32 i = 0; i < ARRAY_COUNT(AtlasSlots); ++i)
+	{
+		FTileAtlasSlot& Slot = AtlasSlots[i];
+		delete Slot.SlotBrush;
+		Slot.SlotBrush = nullptr;
+	}
+	
+	AtlasTexture->RemoveFromRoot();
+	AtlasTexture->MarkPendingKill();
+	AtlasTexture = nullptr;
+}
+
+void FTileAtlasPage::SetOccupied(int32 SlotIdx, bool bOccupied)
+{
+	AtlasSlots[SlotIdx].bOccupied = bOccupied;
+}
+
+bool FTileAtlasPage::HasOccupiedSlots() const
+{
+	for (int32 i = 0; i < ARRAY_COUNT(AtlasSlots); ++i)
+	{
+		if (AtlasSlots[i].bOccupied)
+		{
+			return true;
+		}
+	}
+
+	return false;
+}
+
+int32 FTileAtlasPage::GetFreeSlotIndex() const
+{
+	int32 Result = INDEX_NONE;
+	for (int32 i = 0; i < ARRAY_COUNT(AtlasSlots); ++i)
+	{
+		if (!AtlasSlots[i].bOccupied)
+		{
+			Result = i;
+		}
+	}
+	return Result;
+}
+
+const FSlateBrush* FTileAtlasPage::GetSlotBrush(int32 SlotIdx) const
+{
+	return AtlasSlots[SlotIdx].SlotBrush;
+}
+
+void FTileAtlasPage::UpdateSlotImageData(int32 SlotIdx, FSlateTextureDataPtr ImageData)
+{
+	if (AtlasTexture && AtlasTexture->Resource)
+	{
+		int32 SlotX = (SlotIdx % TileThumbnailAtlasDim)*TileThumbnailSize;
+		int32 SlotY = (SlotIdx / TileThumbnailAtlasDim)*TileThumbnailSize;
+		
+		const FUpdateTextureRegion2D UpdateRegion(
+			SlotX, SlotY,		// Dest X, Y
+			0, 0,				// Source X, Y
+			TileThumbnailSize,	// Width
+			TileThumbnailSize);	// Height
+	
+		struct FSlotUpdateContext
+		{
+			FTextureRHIRef			TextureRHI;
+			FSlateTextureDataPtr	ImageData;
+			uint32					SourcePitch;
+			FUpdateTextureRegion2D	Region;
+		} 
+		Context = 
+		{
+			AtlasTexture->Resource->TextureRHI,
+			ImageData,
+			TileThumbnailSize*4,
+			UpdateRegion
+		};
+			
+		ENQUEUE_UNIQUE_RENDER_COMMAND_ONEPARAMETER(
+			UpdateSlotImageData,
+			FSlotUpdateContext,	Context, Context,
+		{
+			FRHITexture2D* RHITexture2D = (FRHITexture2D*)Context.TextureRHI.GetReference();
+			RHIUpdateTexture2D(RHITexture2D, 0, Context.Region, Context.SourcePitch, Context.ImageData->GetRawBytesPtr());
+		});
+	}
+}
+
+//----------------------------------------------------------------
+//
+//
+//----------------------------------------------------------------
+FTileThumbnailCollection::FTileThumbnailCollection()
+	: LastThumbnailUpdateTime(0.0)
+{
+	SharedThumbnailRT = new FSlateTextureRenderTarget2DResource(
+						FLinearColor::Black, 
+						TileThumbnailSize, 
+						TileThumbnailSize, 
+						PF_B8G8R8A8, SF_Bilinear, TA_Wrap, TA_Wrap, 0.0f
+					);
+
+	BeginInitResource(SharedThumbnailRT);
+}
+
+FTileThumbnailCollection::~FTileThumbnailCollection()
+{
+	for (int32 i = 0; i < AtlasPages.Num(); ++i)
+	{
+		delete AtlasPages[i];
+		AtlasPages[i] = nullptr;
+	}
+
+	TileThumbnailsMap.Reset();
+	
+	BeginReleaseResource(SharedThumbnailRT);
+	FlushRenderingCommands();
+	delete SharedThumbnailRT;
+}
+
+void FTileThumbnailCollection::RegisterTile(const FWorldTileModel& InTileModel)
+{
+	FName TileName = InTileModel.GetLongPackageName();
+	FIntPoint SlotAllocation = AllocateSlot();
+
+	TileThumbnailsMap.Add(TileName, FTileThumbnail(SharedThumbnailRT, InTileModel, SlotAllocation));
+}
+
+void FTileThumbnailCollection::UnregisterTile(const FWorldTileModel& InTileModel)
+{
+	FName TileName = InTileModel.GetLongPackageName();
+	
+	const FTileThumbnail* TileThumbnail = TileThumbnailsMap.Find(TileName);
+	if (TileThumbnail)
+	{
+		ReleaseSlot(TileThumbnail->GetThumbnailSlotAllocation());
+		TileThumbnailsMap.Remove(TileName);
+	}
+}
+
+const FSlateBrush* FTileThumbnailCollection::UpdateTileThumbnail(const FWorldTileModel& InTileModel)
+{
+	FName TileName = InTileModel.GetLongPackageName();
+	FTileThumbnail* TileThumbnail = TileThumbnailsMap.Find(TileName);
+	
+	if (TileThumbnail)
+	{
+		auto ImageData = TileThumbnail->UpdateThumbnail();
+		if (ImageData.IsValid())
+		{
+			LastThumbnailUpdateTime = FPlatformTime::Seconds();
+				
+			FIntPoint SlotAllocation = TileThumbnail->GetThumbnailSlotAllocation();
+			AtlasPages[SlotAllocation.X]->UpdateSlotImageData(SlotAllocation.Y, ImageData);
+			return AtlasPages[SlotAllocation.X]->GetSlotBrush(SlotAllocation.Y);
+		}
+	}
+
+	return &ThumbnailDefaultBrush;
+}
+
+const FSlateBrush* FTileThumbnailCollection::GetTileBrush(const FWorldTileModel& InTileModel) const
+{
+	FName TileName = InTileModel.GetLongPackageName();
+	const FTileThumbnail* TileThumbnail = TileThumbnailsMap.Find(TileName);
+	
+	if (TileThumbnail)
+	{
+		FIntPoint SlotAllocation = TileThumbnail->GetThumbnailSlotAllocation();
+		return AtlasPages[SlotAllocation.X]->GetSlotBrush(SlotAllocation.Y);
+	}
+
+	return &ThumbnailDefaultBrush;
+}
+
+bool FTileThumbnailCollection::IsOnCooldown() const
+{
+	if (TileThumbnailUpdateCooldown > 0.0)
+	{
+		const double CurrentTime = FPlatformTime::Seconds();
+		return ((CurrentTime - LastThumbnailUpdateTime) < TileThumbnailUpdateCooldown);
+	}
+	
+	return false;
+}
+
+FIntPoint FTileThumbnailCollection::AllocateSlot()
+{
+	int32 PageIndex = INDEX_NONE;
+	int32 SlotIndex = INDEX_NONE;
+
+	for (int32 i = 0; i < AtlasPages.Num(); ++i)
+	{
+		SlotIndex = AtlasPages[i]->GetFreeSlotIndex();
+		if (SlotIndex != INDEX_NONE)
+		{
+			PageIndex = i;
+			break;
+		}
+	}
+
+	// Add new page
+	if (SlotIndex == INDEX_NONE)
+	{
+		PageIndex = AtlasPages.Add(new FTileAtlasPage());
+		SlotIndex = AtlasPages[PageIndex]->GetFreeSlotIndex();
+	}
+	
+	check(PageIndex!= INDEX_NONE && SlotIndex != INDEX_NONE);
+	AtlasPages[PageIndex]->SetOccupied(SlotIndex, true);
+
+	return FIntPoint(PageIndex, SlotIndex);
+}
+
+void FTileThumbnailCollection::ReleaseSlot(const FIntPoint& InSlotAllocation)
+{
+	AtlasPages[InSlotAllocation.X]->SetOccupied(InSlotAllocation.Y, false);
+}
diff --git a/Engine/Source/Editor/WorldBrowser/Private/Tiles/WorldTileThumbnails.h b/Engine/Source/Editor/WorldBrowser/Private/Tiles/WorldTileThumbnails.h
new file mode 100644
--- /dev/null
+++ b/Engine/Source/Editor/WorldBrowser/Private/Tiles/WorldTileThumbnails.h
@@ -0,0 +1,97 @@
+// Copyright 1998-2015 Epic Games, Inc. All Rights Reserved.
+
+#pragma once
+
+class FSlateTextureRenderTarget2DResource;
+class FSlateTexture2DRHIRef;
+class FObjectThumbnail;
+struct FSlateDynamicImageBrush;
+class FSlateShaderResourceProxy;
+class FWorldTileModel;
+
+//----------------------------------------------------------------
+//
+//
+//----------------------------------------------------------------
+const int32 TileThumbnailSize = 256;
+const int32 TileThumbnailAtlasSize = 1024;
+const int32 TileThumbnailAtlasDim = TileThumbnailAtlasSize/TileThumbnailSize;
+
+//----------------------------------------------------------------
+//
+//
+//----------------------------------------------------------------
+class FTileThumbnail
+{
+public:
+	FTileThumbnail(FSlateTextureRenderTarget2DResource* InThumbnailRenderTarget, const FWorldTileModel& InTileModel, const FIntPoint& InSlotAllocation);
+
+	/** Redraw thumbnail */
+	FSlateTextureDataPtr UpdateThumbnail();
+	FIntPoint GetThumbnailSlotAllocation() const;
+
+private:
+	FSlateTextureDataPtr ToSlateTextureData(const FObjectThumbnail* ObjectThumbnail) const;
+
+private:
+	const FWorldTileModel&					TileModel;
+	/** Shared render target for slate */
+	FSlateTextureRenderTarget2DResource*	ThumbnailRenderTarget;
+	FIntPoint								SlotAllocation;
+};
+
+//----------------------------------------------------------------
+//
+//
+//----------------------------------------------------------------
+class FTileAtlasPage
+{
+public:
+	FTileAtlasPage();
+	~FTileAtlasPage();
+	
+	void SetOccupied(int32 SlotIdx, bool bOccupied);
+	bool HasOccupiedSlots() const;
+	int32 GetFreeSlotIndex() const;
+	const FSlateBrush* GetSlotBrush(int32 SlotIdx) const;
+	void UpdateSlotImageData(int32 SlotIdx, FSlateTextureDataPtr ImageData);
+		
+private:
+	struct FTileAtlasSlot
+	{
+		FSlateDynamicImageBrush*	SlotBrush;
+		bool						bOccupied;
+	};
+	
+	FTileAtlasSlot		AtlasSlots[TileThumbnailAtlasDim*TileThumbnailAtlasDim];
+	UTexture2DDynamic*	AtlasTexture;
+};
+
+//----------------------------------------------------------------
+//
+//
+//----------------------------------------------------------------
+class FTileThumbnailCollection
+{
+public:
+	FTileThumbnailCollection();
+	~FTileThumbnailCollection();
+
+	void RegisterTile(const FWorldTileModel& InTileModel);
+	void UnregisterTile(const FWorldTileModel& InTileModel);
+	const FSlateBrush* UpdateTileThumbnail(const FWorldTileModel& InTileModel);
+	const FSlateBrush* GetTileBrush(const FWorldTileModel& InTileModel) const;
+
+	bool IsOnCooldown() const;
+
+private:
+	FIntPoint AllocateSlot();
+	void ReleaseSlot(const FIntPoint& InSlotAllocation);
+
+private:
+	FSlateTextureRenderTarget2DResource*	SharedThumbnailRT;
+	TMap<FName, FTileThumbnail>				TileThumbnailsMap;
+	double									LastThumbnailUpdateTime;
+
+	TArray<FTileAtlasPage*>					AtlasPages;
+};
diff --git a/Engine/Source/Editor/WorldBrowser/WorldBrowser.Build.cs b/Engine/Source/Editor/WorldBrowser/WorldBrowser.Build.cs
--- a/Engine/Source/Editor/WorldBrowser/WorldBrowser.Build.cs
+++ b/Engine/Source/Editor/WorldBrowser/WorldBrowser.Build.cs
@@ -44,6 +44,7 @@ public class WorldBrowser : ModuleRules
                 "ImageWrapper",
                 "Foliage",
                 "MaterialUtilities",
+                "RHI"
             }
 		);
 
