Fixed: Heightmaps import flipped on X axis in World Composition Browser
UE-17598

diff --git a/Engine/Source/Editor/WorldBrowser/Private/LevelCollectionModel.h b/Engine/Source/Editor/WorldBrowser/Private/LevelCollectionModel.h
--- a/Engine/Source/Editor/WorldBrowser/Private/LevelCollectionModel.h
+++ b/Engine/Source/Editor/WorldBrowser/Private/LevelCollectionModel.h
@@ -579,6 +579,7 @@ struct FTiledLandscapeImportSettings
 		, SectionsPerComponent(1)
 		, TilesCoordinatesOffset(0,0)
 		, SizeX(1009)
+		, bFlipYAxis(true)
 	{}
 	
 	FVector				Scale3D;
@@ -590,6 +591,7 @@ struct FTiledLandscapeImportSettings
 	TArray<FIntPoint>	TileCoordinates;
 	FIntPoint			TilesCoordinatesOffset;	
 	int32				SizeX;
+	bool				bFlipYAxis;
 
 
 	TWeakObjectPtr<UMaterialInterface>	LandscapeMaterial;
diff --git a/Engine/Source/Editor/WorldBrowser/Private/Tiles/STiledLandscapeImportDlg.cpp b/Engine/Source/Editor/WorldBrowser/Private/Tiles/STiledLandscapeImportDlg.cpp
--- a/Engine/Source/Editor/WorldBrowser/Private/Tiles/STiledLandscapeImportDlg.cpp
+++ b/Engine/Source/Editor/WorldBrowser/Private/Tiles/STiledLandscapeImportDlg.cpp
@@ -73,17 +73,33 @@ void STiledLandcapeImportDlg::Construct(const FArguments& InArgs, TSharedPtr<SWi
 					.OnClicked(this, &STiledLandcapeImportDlg::OnClickedSelectHeightmapTiles)
 					.Text(LOCTEXT("TiledLandscapeImport_SelectButtonText", "Select Heightmap Tiles..."))
 				]
-				
-				// Tiles origin offset
+
+				// Flip Y-Axis orientation
 				+SUniformGridPanel::Slot(0, 1)
 				.VAlign(VAlign_Center)
+				[
+					SNew(STextBlock)
+					.Text( LOCTEXT("TiledLandscapeImport_FlipYAxisText", "Flip Tile Y Coordinate" ) )
+				]
+				+SUniformGridPanel::Slot(1, 1)
+				.VAlign(VAlign_Center)
+				[
+					SNew(SCheckBox)
+					.IsChecked(this, &STiledLandcapeImportDlg::GetFlipYAxisState)
+					.OnCheckStateChanged(this, &STiledLandcapeImportDlg::OnFlipYAxisStateChanged)
+					.ToolTipText(LOCTEXT("TiledLandscapeImport_FlipYAxisToolTip", "Whether tile Y coordinate should be flipped (Make sure 'Flip Y-Axis Orientation' option is switched off in World Machine) "))
+				]
+										
+				// Tiles origin offset
+				+SUniformGridPanel::Slot(0, 2)
+				.VAlign(VAlign_Center)
 				[
 					SNew(STextBlock)
 					.ToolTipText(LOCTEXT("TiledLandscapeImport_TilesOffsetTooltip", "For example: tile x0_y0 will be treated as x(0+offsetX)_y(0+offsetY)"))
 					.Text(LOCTEXT("TiledLandscapeImport_TilesOffsetText", "Tile Coordinates Offset"))
 				]
 
-				+SUniformGridPanel::Slot(1, 1)
+				+SUniformGridPanel::Slot(1, 2)
 				.VAlign(VAlign_Center)
 				[
 					SNew(SHorizontalBox)
@@ -116,16 +132,16 @@ void STiledLandcapeImportDlg::Construct(const FArguments& InArgs, TSharedPtr<SWi
 						]
 					]
 				]
-
+				
 				// Tile configuration
-				+SUniformGridPanel::Slot(0, 2)
+				+SUniformGridPanel::Slot(0, 3)
 				.VAlign(VAlign_Center)
 				[
 					SNew(STextBlock)
 					.Text(LOCTEXT("TiledLandscapeImport_ConfigurationText", "Import Configuration"))
 				]
 
-				+SUniformGridPanel::Slot(1, 2)
+				+SUniformGridPanel::Slot(1, 3)
 				.VAlign(VAlign_Center)
 				[
 					SAssignNew(TileConfigurationComboBox, SComboBox<TSharedPtr<FTileImportConfiguration>>)
@@ -140,14 +156,14 @@ void STiledLandcapeImportDlg::Construct(const FArguments& InArgs, TSharedPtr<SWi
 				]
 
 				// Scale
-				+SUniformGridPanel::Slot(0, 3)
+				+SUniformGridPanel::Slot(0, 4)
 				.VAlign(VAlign_Center)
 				[
 					SNew(STextBlock)
 					.Text(LOCTEXT("TiledLandscapeImport_ScaleText", "Landscape Scale"))
 				]
 			
-				+SUniformGridPanel::Slot(1, 3)
+				+SUniformGridPanel::Slot(1, 4)
 				.VAlign(VAlign_Center)
 				[
 					SNew( SVectorInputBox )
@@ -161,14 +177,14 @@ void STiledLandcapeImportDlg::Construct(const FArguments& InArgs, TSharedPtr<SWi
 				]
 
 				// Landcape material
-				+SUniformGridPanel::Slot(0, 4)
+				+SUniformGridPanel::Slot(0, 5)
 				.VAlign(VAlign_Center)
 				[
 					SNew(STextBlock)
 					.Text(LOCTEXT("TiledLandscapeImport_MaterialText", "Material"))
 				]
 			
-				+SUniformGridPanel::Slot(1, 4)
+				+SUniformGridPanel::Slot(1, 5)
 				.VAlign(VAlign_Center)
 				[
 					SNew(SObjectPropertyEntryBox)
@@ -382,6 +398,16 @@ void STiledLandcapeImportDlg::SetTileOffsetY(int32 InValue)
 	ImportSettings.TilesCoordinatesOffset.Y = InValue;
 }
 
+ECheckBoxState STiledLandcapeImportDlg::GetFlipYAxisState() const
+{
+	return ImportSettings.bFlipYAxis ? ECheckBoxState::Checked : ECheckBoxState::Unchecked;
+}
+
+void STiledLandcapeImportDlg::OnFlipYAxisStateChanged(ECheckBoxState NewState)
+{
+	ImportSettings.bFlipYAxis = (NewState == ECheckBoxState::Checked);
+}
+
 void STiledLandcapeImportDlg::OnSetImportConfiguration(TSharedPtr<FTileImportConfiguration> InTileConfig, ESelectInfo::Type SelectInfo)
 {
 	if (InTileConfig.IsValid())
diff --git a/Engine/Source/Editor/WorldBrowser/Private/Tiles/STiledLandscapeImportDlg.h b/Engine/Source/Editor/WorldBrowser/Private/Tiles/STiledLandscapeImportDlg.h
--- a/Engine/Source/Editor/WorldBrowser/Private/Tiles/STiledLandscapeImportDlg.h
+++ b/Engine/Source/Editor/WorldBrowser/Private/Tiles/STiledLandscapeImportDlg.h
@@ -48,7 +48,11 @@ private:
 	TOptional<int32> GetTileOffsetY() const;
 	void SetTileOffsetX(int32 InValue);
 	void SetTileOffsetY(int32 InValue);
-	
+
+	/** */
+	ECheckBoxState GetFlipYAxisState() const;
+	void OnFlipYAxisStateChanged(ECheckBoxState NewState);
+		
 	/** */
 	void OnSetImportConfiguration(TSharedPtr<FTileImportConfiguration> InTileConfig, ESelectInfo::Type SelectInfo);
 
diff --git a/Engine/Source/Editor/WorldBrowser/Private/Tiles/WorldTileCollectionModel.cpp b/Engine/Source/Editor/WorldBrowser/Private/Tiles/WorldTileCollectionModel.cpp
--- a/Engine/Source/Editor/WorldBrowser/Private/Tiles/WorldTileCollectionModel.cpp
+++ b/Engine/Source/Editor/WorldBrowser/Private/Tiles/WorldTileCollectionModel.cpp
@@ -1639,6 +1639,10 @@ void FWorldTileCollectionModel::ImportTiledLandscape_Executed()
 						float WidthY = NewLandscapeRect.Height()*TileScale.Y;
 						FIntPoint TileCoordinates = ImportSettings.TileCoordinates[TileIndex] + ImportSettings.TilesCoordinatesOffset;
 						FIntPoint TileOffset = FIntPoint(TileCoordinates.X*WidthX, TileCoordinates.Y*WidthY);
+						if (ImportSettings.bFlipYAxis)
+						{
+							TileOffset.Y = -(TileOffset.Y + WidthY);
+						}
 						
 						// Place level tile at correct position in the world
 						NewTileModel->SetLevelPosition(TileOffset);
