Fixed: Merge Actors "place in world" option leaves original actors in world as children UE-15390
Fixed: World composition layers can be named the same UE-18678
Fixed: Level Details summon world composition tooltip is incorrect UE-18674
Added progress dialog to Actor Merging tool

diff --git a/Engine/Source/Developer/MeshUtilities/Private/MeshUtilities.cpp b/Engine/Source/Developer/MeshUtilities/Private/MeshUtilities.cpp
--- a/Engine/Source/Developer/MeshUtilities/Private/MeshUtilities.cpp
+++ b/Engine/Source/Developer/MeshUtilities/Private/MeshUtilities.cpp
@@ -4553,7 +4553,8 @@ void FMeshUtilities::MergeStaticMeshComponents(const TArray<UStaticMeshComponent
 	TArray<FRawMeshExt>								SourceMeshes;
 	bool											bWithVertexColors[MAX_STATIC_MESH_LODS] = {};
 	bool											bOcuppiedUVChannels[MAX_STATIC_MESH_LODS][MAX_MESH_TEXTURE_COORDS] = {};
-
+	FScopedSlowTask									MergeProgress(5);
+	
 	SourceMeshes.Reserve(ComponentsToMerge.Num());
 	SourceMeshes.SetNum(ComponentsToMerge.Num());
 
@@ -4583,6 +4584,8 @@ void FMeshUtilities::MergeStaticMeshComponents(const TArray<UStaticMeshComponent
 		NumMaxLOD = StartLODIndex + 1;
 	}
 
+	MergeProgress.EnterProgressFrame();
+
 	int32 RawMeshLODIdx = 0;
 	for (int32 LODIndex = StartLODIndex; LODIndex < NumMaxLOD; ++LODIndex, ++RawMeshLODIdx)
 	{
@@ -4623,6 +4626,8 @@ void FMeshUtilities::MergeStaticMeshComponents(const TArray<UStaticMeshComponent
 		return;
 	}
 
+	MergeProgress.EnterProgressFrame();
+
 	// For each raw mesh, re-map the material indices according to the MaterialMap
 	for (int32 MeshIndex = 0; MeshIndex < SourceMeshes.Num(); ++MeshIndex)
 	{
@@ -4648,6 +4653,8 @@ void FMeshUtilities::MergeStaticMeshComponents(const TArray<UStaticMeshComponent
 
 	if (InSettings.bMergeMaterials)
 	{
+		MergeProgress.EnterProgressFrame();
+
 		FIntPoint AtlasTextureSize = FIntPoint(InSettings.MergedMaterialAtlasResolution, InSettings.MergedMaterialAtlasResolution);
 		FFlattenMaterial MergedFlatMaterial;
 		MergedFlatMaterial.DiffuseSize = AtlasTextureSize;
@@ -4752,6 +4759,8 @@ void FMeshUtilities::MergeStaticMeshComponents(const TArray<UStaticMeshComponent
 		UniqueMaterials.Add(MergedMaterial);
 	}
 
+	MergeProgress.EnterProgressFrame();
+
 	// Merge meshes into single mesh
 	for (int32 SourceMeshIdx = 0; SourceMeshIdx < SourceMeshes.Num(); ++SourceMeshIdx)
 	{
@@ -4841,6 +4850,8 @@ void FMeshUtilities::MergeStaticMeshComponents(const TArray<UStaticMeshComponent
 		TargetLightMapUVChannel[LODIndex] = FMath::Min(InSettings.TargetLightMapUVChannel, ChannelIdx);
 	}
 
+	MergeProgress.EnterProgressFrame();
+
 	//
 	//Create merged mesh asset
 	//
diff --git a/Engine/Source/Editor/MergeActors/Private/MeshMergingTool/MeshMergingTool.cpp b/Engine/Source/Editor/MergeActors/Private/MeshMergingTool/MeshMergingTool.cpp
--- a/Engine/Source/Editor/MergeActors/Private/MeshMergingTool/MeshMergingTool.cpp
+++ b/Engine/Source/Editor/MergeActors/Private/MeshMergingTool/MeshMergingTool.cpp
@@ -19,7 +19,7 @@
 
 
 FMeshMergingTool::FMeshMergingTool()
-	: bPlaceInWorld(false)
+	: bReplaceSourceActors(false)
 	, bExportSpecificLOD(false)
 	, ExportLODIndex(0)
 {}
@@ -81,7 +81,7 @@ bool FMeshMergingTool::RunMerge(const FString& PackageName)
 	}
 
 	// This restriction is only for replacement of selected actors with merged mesh actor
-	if (UniqueLevels.Num() > 1 && bPlaceInWorld)
+	if (UniqueLevels.Num() > 1 && bReplaceSourceActors)
 	{
 		FText Message = NSLOCTEXT("UnrealEd", "FailedToMergeActorsSublevels_Msg", "The selected actors should be in the same level");
 		OpenMsgDlgInt(EAppMsgType::Ok, Message, NSLOCTEXT("UnrealEd", "FailedToMergeActors_Title", "Unable to merge actors"));
@@ -89,10 +89,15 @@ bool FMeshMergingTool::RunMerge(const FString& PackageName)
 	}
 
 	int32 TargetMeshLOD = bExportSpecificLOD ? ExportLODIndex : INDEX_NONE;
-
 	FVector MergedActorLocation;
 	TArray<UObject*> AssetsToSync;
-	MeshUtilities.MergeActors(Actors, MergingSettings, NULL, PackageName, TargetMeshLOD, AssetsToSync, MergedActorLocation);
+	// Merge...
+	{
+		FScopedSlowTask SlowTask(0, LOCTEXT("MergingActorsSlowTask", "Merging actors..."));
+		SlowTask.MakeDialog();
+
+		MeshUtilities.MergeActors(Actors, MergingSettings, NULL, PackageName, TargetMeshLOD, AssetsToSync, MergedActorLocation);
+	}
 
 	if (AssetsToSync.Num())
 	{
@@ -109,12 +114,12 @@ bool FMeshMergingTool::RunMerge(const FString& PackageName)
 		ContentBrowserModule.Get().SyncBrowserToAssets(AssetsToSync, true);
 
 		// Place new mesh in the world
-		if (bPlaceInWorld)
+		if (bReplaceSourceActors)
 		{
 			UStaticMesh* MergedMesh = nullptr;
 			if (AssetsToSync.FindItemByClass(&MergedMesh))
 			{
-				const FScopedTransaction Transaction(NSLOCTEXT("UnrealEd", "PlaceMergedActor", "Place Merged Actor"));
+				const FScopedTransaction Transaction(LOCTEXT("PlaceMergedActor", "Place Merged Actor"));
 				UniqueLevels[0]->Modify();
 
 				UWorld* World = UniqueLevels[0]->OwningWorld;
@@ -126,13 +131,10 @@ bool FMeshMergingTool::RunMerge(const FString& PackageName)
 				MergedActor->GetStaticMeshComponent()->StaticMesh = MergedMesh;
 				MergedActor->SetActorLabel(AssetsToSync[0]->GetName());
 
-				// Add source actors as children to merged actor and hide them
+				// Remove source actors
 				for (AActor* Actor : Actors)
 				{
-					Actor->Modify();
-					Actor->AttachRootComponentToActor(MergedActor, NAME_None, EAttachLocation::KeepWorldPosition);
-					Actor->SetActorHiddenInGame(true);
-					Actor->SetIsTemporarilyHiddenInEditor(true);
+					Actor->Destroy();
 				}
 			}
 		}
diff --git a/Engine/Source/Editor/MergeActors/Private/MeshMergingTool/MeshMergingTool.h b/Engine/Source/Editor/MergeActors/Private/MeshMergingTool/MeshMergingTool.h
--- a/Engine/Source/Editor/MergeActors/Private/MeshMergingTool/MeshMergingTool.h
+++ b/Engine/Source/Editor/MergeActors/Private/MeshMergingTool/MeshMergingTool.h
@@ -27,8 +27,8 @@ private:
 	/** Current mesh merging settings */
 	FMeshMergingSettings MergingSettings;
 
-	/** Whether to spawn merged actor in the world */
-	bool bPlaceInWorld;
+	/** Whether to replace source actors with a merged actor in the world */
+	bool bReplaceSourceActors;
 
 	bool bExportSpecificLOD;
 	int32 ExportLODIndex;
diff --git a/Engine/Source/Editor/MergeActors/Private/MeshMergingTool/SMeshMergingDialog.cpp b/Engine/Source/Editor/MergeActors/Private/MeshMergingTool/SMeshMergingDialog.cpp
--- a/Engine/Source/Editor/MergeActors/Private/MeshMergingTool/SMeshMergingDialog.cpp
+++ b/Engine/Source/Editor/MergeActors/Private/MeshMergingTool/SMeshMergingDialog.cpp
@@ -215,19 +215,19 @@ void SMeshMergingDialog::Construct(const FArguments& InArgs, FMeshMergingTool* I
 					]
 				]
 
-				// Place in world
+				// Replace source actors
 				+SVerticalBox::Slot()
 				.AutoHeight()
 				.Padding(FEditorStyle::GetMargin("StandardDialog.ContentPadding"))
 				[
 					SNew(SCheckBox)
 					.Type(ESlateCheckBoxType::CheckBox)
-					.IsChecked(this, &SMeshMergingDialog::GetPlaceInWorld)
-					.OnCheckStateChanged(this, &SMeshMergingDialog::SetPlaceInWorld)
+					.IsChecked(this, &SMeshMergingDialog::GetReplaceSourceActors)
+					.OnCheckStateChanged(this, &SMeshMergingDialog::SetReplaceSourceActors)
 					.Content()
 					[
 						SNew(STextBlock)
-						.Text(LOCTEXT("PlaceInWorldLabel", "Place In World"))
+						.Text(LOCTEXT("ReplaceSourceActorsLabel", "Replace Source Actors"))
 						.Font(FEditorStyle::GetFontStyle("StandardDialog.SmallFont"))
 					]
 				]
@@ -417,14 +417,14 @@ void SMeshMergingDialog::SetPivotPointAtZero(ECheckBoxState NewValue)
 	Tool->MergingSettings.bPivotPointAtZero = (ECheckBoxState::Checked == NewValue);
 }
 
-ECheckBoxState SMeshMergingDialog::GetPlaceInWorld() const
+ECheckBoxState SMeshMergingDialog::GetReplaceSourceActors() const
 {
-	return (Tool->bPlaceInWorld ? ECheckBoxState::Checked : ECheckBoxState::Unchecked);
+	return (Tool->bReplaceSourceActors ? ECheckBoxState::Checked : ECheckBoxState::Unchecked);
 }
 
-void SMeshMergingDialog::SetPlaceInWorld(ECheckBoxState NewValue)
+void SMeshMergingDialog::SetReplaceSourceActors(ECheckBoxState NewValue)
 {
-	Tool->bPlaceInWorld = (ECheckBoxState::Checked == NewValue);
+	Tool->bReplaceSourceActors = (ECheckBoxState::Checked == NewValue);
 }
 
 bool SMeshMergingDialog::IsMaterialMergingEnabled() const
diff --git a/Engine/Source/Editor/MergeActors/Private/MeshMergingTool/SMeshMergingDialog.h b/Engine/Source/Editor/MergeActors/Private/MeshMergingTool/SMeshMergingDialog.h
--- a/Engine/Source/Editor/MergeActors/Private/MeshMergingTool/SMeshMergingDialog.h
+++ b/Engine/Source/Editor/MergeActors/Private/MeshMergingTool/SMeshMergingDialog.h
@@ -50,8 +50,8 @@ private:
 	ECheckBoxState GetPivotPointAtZero() const;
 	void SetPivotPointAtZero(ECheckBoxState NewValue);
 
-	ECheckBoxState GetPlaceInWorld() const;
-	void SetPlaceInWorld(ECheckBoxState NewValue);
+	ECheckBoxState GetReplaceSourceActors() const;
+	void SetReplaceSourceActors(ECheckBoxState NewValue);
 
 	/** Material merging */
 	bool IsMaterialMergingEnabled() const;
diff --git a/Engine/Source/Editor/WorldBrowser/Private/SWorldDetails.cpp b/Engine/Source/Editor/WorldBrowser/Private/SWorldDetails.cpp
--- a/Engine/Source/Editor/WorldBrowser/Private/SWorldDetails.cpp
+++ b/Engine/Source/Editor/WorldBrowser/Private/SWorldDetails.cpp
@@ -130,7 +130,7 @@ void SWorldDetails::OnBrowseWorld(UWorld* InWorld)
 						.Visibility(this, &SWorldDetails::GetCompositionButtonVisibility)
 						.ButtonStyle(FEditorStyle::Get(), "ToggleButton")
 						.OnClicked(this, &SWorldDetails::OnSummonComposition)
-						.ToolTipText(LOCTEXT("SummonHierarchyToolTipText", "Summons world composition"))
+						.ToolTipText(LOCTEXT("SummonCompositionToolTipText", "Summons world composition"))
 						.HAlign(HAlign_Center)
 						.VAlign(VAlign_Center)
 						.Content()
diff --git a/Engine/Source/Editor/WorldBrowser/Private/SWorldHierarchy.cpp b/Engine/Source/Editor/WorldBrowser/Private/SWorldHierarchy.cpp
--- a/Engine/Source/Editor/WorldBrowser/Private/SWorldHierarchy.cpp
+++ b/Engine/Source/Editor/WorldBrowser/Private/SWorldHierarchy.cpp
@@ -702,7 +702,7 @@ void SWorldHierarchy::OnBrowseWorld(UWorld* InWorld)
 						.Visibility(this, &SWorldHierarchy::GetCompositionButtonVisibility)
 						.ButtonStyle(FEditorStyle::Get(), "ToggleButton")
 						.OnClicked(this, &SWorldHierarchy::OnSummonComposition)
-						.ToolTipText(LOCTEXT("SummonHierarchyToolTipText", "Summons world composition"))
+						.ToolTipText(LOCTEXT("SummonCompositionToolTipText", "Summons world composition"))
 						.HAlign(HAlign_Center)
 						.VAlign(VAlign_Center)
 						.Content()
diff --git a/Engine/Source/Editor/WorldBrowser/Private/Tiles/SWorldComposition.cpp b/Engine/Source/Editor/WorldBrowser/Private/Tiles/SWorldComposition.cpp
--- a/Engine/Source/Editor/WorldBrowser/Private/Tiles/SWorldComposition.cpp
+++ b/Engine/Source/Editor/WorldBrowser/Private/Tiles/SWorldComposition.cpp
@@ -1200,10 +1200,11 @@ FReply SWorldComposition::NewLayer_Clicked()
 		return FReply::Handled();
 	}
 	
-	TSharedRef<SNewLayerPopup> CreateLayerWidget = 
-		SNew(SNewLayerPopup)
+	TSharedRef<SNewWorldLayerPopup> CreateLayerWidget = 
+		SNew(SNewWorldLayerPopup)
 		.OnCreateLayer(this, &SWorldComposition::CreateNewLayer)
-		.DefaultName(LOCTEXT("Layer_DefaultName", "MyLayer").ToString());
+		.DefaultName(LOCTEXT("Layer_DefaultName", "MyLayer").ToString())
+		.InWorldModel(TileWorldModel);
 
 	NewLayerMenu = FSlateApplication::Get().PushMenu(
 		this->AsShared(),
diff --git a/Engine/Source/Editor/WorldBrowser/Private/Tiles/SWorldLayers.cpp b/Engine/Source/Editor/WorldBrowser/Private/Tiles/SWorldLayers.cpp
--- a/Engine/Source/Editor/WorldBrowser/Private/Tiles/SWorldLayers.cpp
+++ b/Engine/Source/Editor/WorldBrowser/Private/Tiles/SWorldLayers.cpp
@@ -12,10 +12,19 @@
 //
 //----------------------------------------------------------------
 BEGIN_SLATE_FUNCTION_BUILD_OPTIMIZATION
-void SNewLayerPopup::Construct(const FArguments& InArgs)
+void SNewWorldLayerPopup::Construct(const FArguments& InArgs)
 {
-	OnCreateLayer = InArgs._OnCreateLayer;
-	LayerData.Name = InArgs._DefaultName;
+	OnCreateLayer	= InArgs._OnCreateLayer;
+	LayerData.Name	= InArgs._DefaultName;
+
+	// store set of currently existing layer names
+	{
+		const auto& AllLayersList = InArgs._InWorldModel->GetLayers();
+		for (const auto& WorldLayer : AllLayersList)
+		{
+			ExistingLayerNames.Add(WorldLayer.Name);
+		}
+	}
 	
 	ChildSlot
 	[
@@ -42,10 +51,9 @@ void SNewLayerPopup::Construct(const FArguments& InArgs)
 				.Padding(4,0,0,0)
 				[
 					SNew(SEditableTextBox)
-					.Text(this, &SNewLayerPopup::GetLayerName)
+					.Text(this, &SNewWorldLayerPopup::GetLayerName)
 					.SelectAllTextWhenFocused(true)
-					//.OnTextCommitted(this, &SNewLayerPopup::SetLayerName)
-					.OnTextChanged(this, &SNewLayerPopup::SetLayerName)
+					.OnTextChanged(this, &SNewWorldLayerPopup::SetLayerName)
 				]
 
 			]
@@ -60,19 +68,19 @@ void SNewLayerPopup::Construct(const FArguments& InArgs)
 				.AutoWidth()
 				[
 					SNew(SCheckBox)
-					.IsChecked(this, &SNewLayerPopup::GetDistanceStreamingState)
-					.OnCheckStateChanged(this, &SNewLayerPopup::OnDistanceStreamingStateChanged)
+					.IsChecked(this, &SNewWorldLayerPopup::GetDistanceStreamingState)
+					.OnCheckStateChanged(this, &SNewWorldLayerPopup::OnDistanceStreamingStateChanged)
 				]
 
 				+SHorizontalBox::Slot()
 				.AutoWidth()
 				[
 					SNew(SNumericEntryBox<int32>)
-					.IsEnabled(this, &SNewLayerPopup::IsDistanceStreamingEnabled)
-					.Value(this, &SNewLayerPopup::GetStreamingDistance)
+					.IsEnabled(this, &SNewWorldLayerPopup::IsDistanceStreamingEnabled)
+					.Value(this, &SNewWorldLayerPopup::GetStreamingDistance)
 					.MinValue(1)
 					.MaxValue(TNumericLimits<int32>::Max())
-					.OnValueChanged(this, &SNewLayerPopup::SetStreamingDistance)
+					.OnValueChanged(this, &SNewWorldLayerPopup::SetStreamingDistance)
 					.LabelPadding(0)
 					.Label()
 					[
@@ -90,7 +98,8 @@ void SNewLayerPopup::Construct(const FArguments& InArgs)
 			.Padding(2,2,0,0)
 			[
 				SNew(SButton)
-				.OnClicked(this, &SNewLayerPopup::OnClickedCreate)
+				.OnClicked(this, &SNewWorldLayerPopup::OnClickedCreate)
+				.IsEnabled(this, &SNewWorldLayerPopup::CanCreateLayer)
 				.Text(LOCTEXT("Layer_Create", "Create"))
 			]
 
@@ -99,7 +108,7 @@ void SNewLayerPopup::Construct(const FArguments& InArgs)
 }
 END_SLATE_FUNCTION_BUILD_OPTIMIZATION
 
-FReply SNewLayerPopup::OnClickedCreate()
+FReply SNewWorldLayerPopup::OnClickedCreate()
 {
 	if (OnCreateLayer.IsBound())
 	{
@@ -109,6 +118,11 @@ FReply SNewLayerPopup::OnClickedCreate()
 	return FReply::Unhandled();
 }
 
+bool SNewWorldLayerPopup::CanCreateLayer() const
+{
+	return LayerData.Name.Len() > 0 && !ExistingLayerNames.Contains(LayerData.Name);
+}
+
 /** A class for check boxes in the layer list. 
  *	If you double click a layer checkbox, you will enable it and disable all others 
  *	If you Ctrl+Click a layer checkbox, you will add/remove it from selection list
diff --git a/Engine/Source/Editor/WorldBrowser/Private/Tiles/SWorldLayers.h b/Engine/Source/Editor/WorldBrowser/Private/Tiles/SWorldLayers.h
--- a/Engine/Source/Editor/WorldBrowser/Private/Tiles/SWorldLayers.h
+++ b/Engine/Source/Editor/WorldBrowser/Private/Tiles/SWorldLayers.h
@@ -34,16 +34,17 @@ private:
 };
 
 
-class SNewLayerPopup 
+class SNewWorldLayerPopup 
 	: public SBorder
 {
 public:
 	DECLARE_DELEGATE_RetVal_OneParam(FReply, FOnCreateLayer, const FWorldTileLayer&)
 	
-	SLATE_BEGIN_ARGS(SNewLayerPopup)
+	SLATE_BEGIN_ARGS(SNewWorldLayerPopup)
 	{}
 	SLATE_EVENT(FOnCreateLayer, OnCreateLayer)
 	SLATE_ARGUMENT(FString, DefaultName)
+	SLATE_ARGUMENT(TSharedPtr<FWorldTileCollectionModel>, InWorldModel)
 	SLATE_END_ARGS()
 		
 	void Construct(const FArguments& InArgs);
@@ -75,6 +76,7 @@ public:
 	
 private:
 	FReply OnClickedCreate();
+	bool CanCreateLayer() const;
 
 	void SetLayerName(const FText& InText)
 	{
@@ -93,6 +95,7 @@ private:
 
 private:
 	/** The delegate to execute when the create button is clicked */
-	FOnCreateLayer			OnCreateLayer;
-	FWorldTileLayer			LayerData;
+	FOnCreateLayer							OnCreateLayer;
+	FWorldTileLayer							LayerData;
+	TSet<FString>							ExistingLayerNames;
 };
