World Browser: Do not cache pointer to UEditorEngine use global instead
UE-8839

diff --git a/Engine/Source/Editor/WorldBrowser/Private/LevelCollectionModel.cpp b/Engine/Source/Editor/WorldBrowser/Private/LevelCollectionModel.cpp
--- a/Engine/Source/Editor/WorldBrowser/Private/LevelCollectionModel.cpp
+++ b/Engine/Source/Editor/WorldBrowser/Private/LevelCollectionModel.cpp
@@ -13,9 +13,8 @@
 
 #define LOCTEXT_NAMESPACE "WorldBrowser"
 
-FLevelCollectionModel::FLevelCollectionModel(UEditorEngine* InEditor)
-	: Editor(InEditor)
-	, bRequestedUpdateAllLevels(false)
+FLevelCollectionModel::FLevelCollectionModel()
+	: bRequestedUpdateAllLevels(false)
 	, bRequestedRedrawAllLevels(false)
 	, bRequestedUpdateActorsCount(false)
 	, CommandList(MakeShareable(new FUICommandList))
@@ -39,8 +38,8 @@ FLevelCollectionModel::~FLevelCollectionModel()
 	FWorldDelegates::LevelAddedToWorld.RemoveAll(this);
 	FWorldDelegates::LevelRemovedFromWorld.RemoveAll(this);
 	FEditorSupportDelegates::RedrawAllViewports.RemoveAll(this);
-	Editor->OnLevelActorAdded().RemoveAll(this);
-	Editor->OnLevelActorDeleted().RemoveAll(this);
+	GEditor->OnLevelActorAdded().RemoveAll(this);
+	GEditor->OnLevelActorDeleted().RemoveAll(this);
 	if (CurrentWorld.IsValid())
 	{
 		CurrentWorld->OnSelectedLevelsChanged().RemoveAll(this);
@@ -57,8 +56,8 @@ void FLevelCollectionModel::Initialize(UWorld* InWorld)
 	FWorldDelegates::LevelAddedToWorld.AddSP(this, &FLevelCollectionModel::OnLevelAddedToWorld);
 	FWorldDelegates::LevelRemovedFromWorld.AddSP(this, &FLevelCollectionModel::OnLevelRemovedFromWorld);
 	FEditorSupportDelegates::RedrawAllViewports.AddSP(this, &FLevelCollectionModel::OnRedrawAllViewports);
-	Editor->OnLevelActorAdded().AddSP( this, &FLevelCollectionModel::OnLevelActorAdded);
-	Editor->OnLevelActorDeleted().AddSP( this, &FLevelCollectionModel::OnLevelActorDeleted);
+	GEditor->OnLevelActorAdded().AddSP( this, &FLevelCollectionModel::OnLevelActorAdded);
+	GEditor->OnLevelActorDeleted().AddSP( this, &FLevelCollectionModel::OnLevelActorDeleted);
 	USelection::SelectionChangedEvent.AddSP(this, &FLevelCollectionModel::OnActorSelectionChanged);
 	SelectionChanged.AddSP(this, &FLevelCollectionModel::OnActorOrLevelSelectionChanged);
 	CurrentWorld->OnSelectedLevelsChanged().AddSP(this, &FLevelCollectionModel::OnLevelsSelectionChangedOutside);
@@ -229,6 +228,12 @@ void FLevelCollectionModel::PopulateFilteredLevelsList()
 
 void FLevelCollectionModel::Tick( float DeltaTime )
 {
+	if (GEditor == nullptr)
+	{
+		// Could it be during hot-reloading?
+		return;
+	}
+	
 	if (bRequestedUpdateAllLevels)
 	{
 		UpdateAllLevels();
@@ -282,7 +287,12 @@ bool FLevelCollectionModel::IsReadOnly() const
 
 bool FLevelCollectionModel::IsSimulating() const
 {
-	return (Editor->bIsSimulatingInEditor || Editor->PlayWorld != NULL);
+	return (GEditor->bIsSimulatingInEditor || GEditor->PlayWorld != NULL);
+}
+
+UWorld* FLevelCollectionModel::GetSimulationWorld() const 
+{ 
+	return GEditor->PlayWorld; 
 }
 
 FLevelModelList& FLevelCollectionModel::GetRootLevelList()
@@ -582,7 +592,7 @@ void FLevelCollectionModel::UnloadLevels(const FLevelModelList& InLevelList)
 			// This avoids crashing in areas that rely on getting a selected actors level. The level will be invalid after its removed.
 			for (auto ActorIt = Level->Actors.CreateIterator(); ActorIt; ++ActorIt)
 			{
-				Editor->SelectActor((*ActorIt), /*bInSelected=*/ false, /*bSelectEvenIfHidden=*/ false);
+				GEditor->SelectActor((*ActorIt), /*bInSelected=*/ false, /*bSelectEvenIfHidden=*/ false);
 			}
 			
 			{
@@ -592,7 +602,7 @@ void FLevelCollectionModel::UnloadLevels(const FLevelModelList& InLevelList)
 		}
 	}
 
-	Editor->ResetTransaction( LOCTEXT("RemoveLevelTransReset", "Removing Levels from World") );
+	GEditor->ResetTransaction( LOCTEXT("RemoveLevelTransReset", "Removing Levels from World") );
 
 	// Collect garbage to clear out the destroyed level
 	CollectGarbage( GARBAGE_COLLECTION_KEEPFLAGS );
@@ -896,7 +906,7 @@ bool FLevelCollectionModel::AreAnySelectedLevelsDirty() const
 
 bool FLevelCollectionModel::AreActorsSelected() const
 {
-	return Editor->GetSelectedActorCount() > 0;
+	return GEditor->GetSelectedActorCount() > 0;
 }
 
 bool FLevelCollectionModel::GetDisplayPathsState() const
@@ -1349,7 +1359,7 @@ void FLevelCollectionModel::OnSCCRefresh()
 
 void FLevelCollectionModel::OnSCCDiffAgainstDepot()
 {
-	SCCDiffAgainstDepot(GetSelectedLevels(), Editor.Get());
+	SCCDiffAgainstDepot(GetSelectedLevels(), GEditor);
 }
 
 void FLevelCollectionModel::OnSCCConnect() const
@@ -1461,7 +1471,7 @@ void FLevelCollectionModel::UnockAllLevels_Executed()
 
 void FLevelCollectionModel::ToggleReadOnlyLevels_Executed()
 {
-	Editor->bLockReadOnlyLevels = !Editor->bLockReadOnlyLevels;
+	GEditor->bLockReadOnlyLevels = !GEditor->bLockReadOnlyLevels;
 }
 
 void FLevelCollectionModel::MakeLevelCurrent_Executed()
@@ -1502,7 +1512,7 @@ void FLevelCollectionModel::MoveActorsToSelected_Executed()
 
 	MakeLevelCurrent_Executed();
 	const FScopedTransaction Transaction( LOCTEXT("MoveSelectedActorsToSelectedLevel", "Move Selected Actors to Level") );
-	Editor->MoveSelectedActorsToLevel(GetWorld()->GetCurrentLevel());
+	GEditor->MoveSelectedActorsToLevel(GetWorld()->GetCurrentLevel());
 
 	RequestUpdateAllLevels();
 }
@@ -1512,7 +1522,7 @@ void FLevelCollectionModel::MoveFoliageToSelected_Executed()
 	if (GetSelectedLevels().Num() == 1)
 	{
 		ULevel* TargetLevel = GetSelectedLevels()[0]->GetLevelObject();
-		Editor->MoveSelectedFoliageToLevel(TargetLevel);
+		GEditor->MoveSelectedFoliageToLevel(TargetLevel);
 	}
 }
 
@@ -1520,8 +1530,8 @@ void FLevelCollectionModel::SelectActors_Executed()
 {
 	//first clear any existing actor selection
 	const FScopedTransaction Transaction( LOCTEXT("SelectActors", "Select Actors in Level") );
-	Editor->GetSelectedActors()->Modify();
-	Editor->SelectNone( false, true );
+	GEditor->GetSelectedActors()->Modify();
+	GEditor->SelectNone( false, true );
 
 	for(auto It = SelectedLevelsList.CreateConstIterator(); It; ++It)
 	{
@@ -1564,7 +1574,7 @@ void FLevelCollectionModel::FillLockSubMenu(FMenuBuilder& InMenuBuilder)
 	InMenuBuilder.AddMenuEntry( Commands.World_LockAllLevels );
 	InMenuBuilder.AddMenuEntry( Commands.World_UnockAllLevels );
 
-	if (Editor->bLockReadOnlyLevels)
+	if (GEditor->bLockReadOnlyLevels)
 	{
 		InMenuBuilder.AddMenuEntry( Commands.World_UnlockReadOnlyLevels );
 	}
diff --git a/Engine/Source/Editor/WorldBrowser/Private/LevelCollectionModel.h b/Engine/Source/Editor/WorldBrowser/Private/LevelCollectionModel.h
--- a/Engine/Source/Editor/WorldBrowser/Private/LevelCollectionModel.h
+++ b/Engine/Source/Editor/WorldBrowser/Private/LevelCollectionModel.h
@@ -22,7 +22,7 @@ public:
 
 
 public:
-	FLevelCollectionModel(UEditorEngine* InEditor);
+	FLevelCollectionModel();
 	virtual ~FLevelCollectionModel();
 
 	/** FTickableEditorObject interface */
@@ -38,7 +38,7 @@ public:
 	bool IsSimulating() const;
 
 	/**	@return	Current simulation world */
-	UWorld* GetSimulationWorld() const { return Editor->PlayWorld; }
+	UWorld* GetSimulationWorld() const;
 
 	/**	@return	Current editor world */
 	UWorld* GetWorld(bool bEvenIfPendingKill = false) const { return CurrentWorld.Get(bEvenIfPendingKill); }
@@ -471,8 +471,6 @@ protected:
 	void CacheCanExecuteSourceControlVars() const;
 		
 protected:
-	//
-	const TWeakObjectPtr<UEditorEngine> Editor;
 	
 	// The editor world from where we pull our data
 	TWeakObjectPtr<UWorld>				CurrentWorld;
diff --git a/Engine/Source/Editor/WorldBrowser/Private/LevelModel.cpp b/Engine/Source/Editor/WorldBrowser/Private/LevelModel.cpp
--- a/Engine/Source/Editor/WorldBrowser/Private/LevelModel.cpp
+++ b/Engine/Source/Editor/WorldBrowser/Private/LevelModel.cpp
@@ -12,10 +12,8 @@
 
 #define LOCTEXT_NAMESPACE "WorldBrowser"
 
-FLevelModel::FLevelModel(FLevelCollectionModel& InLevelCollectionModel,
-						const TWeakObjectPtr<UEditorEngine>& InEditor)
+FLevelModel::FLevelModel(FLevelCollectionModel& InLevelCollectionModel)
 	: LevelCollectionModel(InLevelCollectionModel)
-	, Editor(InEditor)
 	, bSelected(false)
 	, bExpanded(false)
 	, bLoadingLevel(false)
@@ -250,7 +248,7 @@ void FLevelModel::SetLocked(bool bLocked)
 		DeselectAllSurfaces();
 
 		// Tell the editor selection status was changed.
-		Editor->NoteSelectionChange();
+		GEditor->NoteSelectionChange();
 
 		// If locking the current level, reset the p-level as the current level
 		//@todo: fix this!
@@ -290,7 +288,7 @@ void FLevelModel::MakeLevelCurrent()
 				
 			// Deselect all selected builder brushes.
 			bool bDeselectedSomething = false;
-			for (FSelectionIterator It(Editor->GetSelectedActorIterator()); It; ++It)
+			for (FSelectionIterator It(GEditor->GetSelectedActorIterator()); It; ++It)
 			{
 				AActor* Actor = static_cast<AActor*>(*It);
 				checkSlow(Actor->IsA(AActor::StaticClass()));
@@ -298,7 +296,7 @@ void FLevelModel::MakeLevelCurrent()
 				ABrush* Brush = Cast< ABrush >( Actor );
 				if (Brush && FActorEditorUtils::IsABuilderBrush(Brush))
 				{
-					Editor->SelectActor(Actor, /*bInSelected=*/ false, /*bNotify=*/ false);
+					GEditor->SelectActor(Actor, /*bInSelected=*/ false, /*bNotify=*/ false);
 					bDeselectedSomething = true;
 				}
 			}
@@ -306,7 +304,7 @@ void FLevelModel::MakeLevelCurrent()
 			// Send a selection change callback if necessary.
 			if (bDeselectedSomething)
 			{
-				Editor->NoteSelectionChange();
+				GEditor->NoteSelectionChange();
 			}
 		}
 							
@@ -613,7 +611,7 @@ void FLevelModel::DeselectAllActors()
 		return;
 	}
 
-	USelection* SelectedActors = Editor->GetSelectedActors();
+	USelection* SelectedActors = GEditor->GetSelectedActors();
 	SelectedActors->Modify();
 
 	// Deselect all level actors 
@@ -642,7 +640,7 @@ void FLevelModel::SelectActors(bool bSelect, bool bNotify, bool bSelectEvenIfHid
 		return;
 	}
 
-	Editor->GetSelectedActors()->BeginBatchSelectOperation();
+	GEditor->GetSelectedActors()->BeginBatchSelectOperation();
 	bool bChangesOccurred = false;
 
 	// Iterate over all actors, looking for actors in this level.
@@ -665,17 +663,17 @@ void FLevelModel::SelectActors(bool bSelect, bool bNotify, bool bSelectEvenIfHid
 			}
 
 			bool bNotifyForActor = false;
-			Editor->GetSelectedActors()->Modify();
-			Editor->SelectActor(Actor, bSelect, bNotifyForActor, bSelectEvenIfHidden);
+			GEditor->GetSelectedActors()->Modify();
+			GEditor->SelectActor(Actor, bSelect, bNotifyForActor, bSelectEvenIfHidden);
 			bChangesOccurred = true;
 		}
 	}
 
-	Editor->GetSelectedActors()->EndBatchSelectOperation();
+	GEditor->GetSelectedActors()->EndBatchSelectOperation();
 
 	if (bNotify)
 	{
-		Editor->NoteSelectionChange();
+		GEditor->NoteSelectionChange();
 	}
 }
 
diff --git a/Engine/Source/Editor/WorldBrowser/Private/LevelModel.h b/Engine/Source/Editor/WorldBrowser/Private/LevelModel.h
--- a/Engine/Source/Editor/WorldBrowser/Private/LevelModel.h
+++ b/Engine/Source/Editor/WorldBrowser/Private/LevelModel.h
@@ -26,8 +26,7 @@ public:
 	DECLARE_EVENT( FLevelModel, FSimpleEvent );
 	
 public:
-	FLevelModel(FLevelCollectionModel& InLevelCollectionModel,
-				const TWeakObjectPtr< UEditorEngine >& InEditor);
+	FLevelModel(FLevelCollectionModel& InLevelCollectionModel);
 
 	virtual ~FLevelModel();
 	
@@ -271,9 +270,6 @@ protected:
 	/** Reference to owning collection model */
 	FLevelCollectionModel&				LevelCollectionModel;
 		
-	/** The UEditorEngine to use */
-	const TWeakObjectPtr<UEditorEngine> Editor;
-
 	/** The parent level  */
 	TWeakPtr<FLevelModel>				Parent;
 		
diff --git a/Engine/Source/Editor/WorldBrowser/Private/StreamingLevels/StreamingLevelCollectionModel.cpp b/Engine/Source/Editor/WorldBrowser/Private/StreamingLevels/StreamingLevelCollectionModel.cpp
--- a/Engine/Source/Editor/WorldBrowser/Private/StreamingLevels/StreamingLevelCollectionModel.cpp
+++ b/Engine/Source/Editor/WorldBrowser/Private/StreamingLevels/StreamingLevelCollectionModel.cpp
@@ -18,8 +18,8 @@
 
 #define LOCTEXT_NAMESPACE "WorldBrowser"
 
-FStreamingLevelCollectionModel::FStreamingLevelCollectionModel(UEditorEngine* InEditor)
-	: FLevelCollectionModel(InEditor)
+FStreamingLevelCollectionModel::FStreamingLevelCollectionModel()
+	: FLevelCollectionModel()
 	, AddedLevelStreamingClass(ULevelStreamingKismet::StaticClass())
 	, bAssetDialogOpen(false)
 {
@@ -32,13 +32,13 @@ FStreamingLevelCollectionModel::FStreamingLevelCollectionModel(UEditorEngine* In
 
 FStreamingLevelCollectionModel::~FStreamingLevelCollectionModel()
 {
-	Editor->UnregisterForUndo( this );
+	GEditor->UnregisterForUndo( this );
 }
 
 void FStreamingLevelCollectionModel::Initialize(UWorld* InWorld)
 {
 	BindCommands();	
-	Editor->RegisterForUndo( this );
+	GEditor->RegisterForUndo( this );
 	
 	FLevelCollectionModel::Initialize(InWorld);
 }
@@ -54,7 +54,7 @@ void FStreamingLevelCollectionModel::OnLevelsCollectionChanged()
 	}
 
 	// Add model for a persistent level
-	TSharedPtr<FStreamingLevelModel> PersistentLevelModel = MakeShareable(new FStreamingLevelModel(Editor, *this, NULL));
+	TSharedPtr<FStreamingLevelModel> PersistentLevelModel = MakeShareable(new FStreamingLevelModel(*this, nullptr));
 	PersistentLevelModel->SetLevelExpansionFlag(true);
 	RootLevelsList.Add(PersistentLevelModel);
 	AllLevelsList.Add(PersistentLevelModel);
@@ -64,9 +64,9 @@ void FStreamingLevelCollectionModel::OnLevelsCollectionChanged()
 	for (auto It = CurrentWorld->StreamingLevels.CreateConstIterator(); It; ++It)
 	{
 		ULevelStreaming* StreamingLevel = (*It);
-		if (StreamingLevel != NULL)
+		if (StreamingLevel != nullptr)
 		{
-			TSharedPtr<FStreamingLevelModel> LevelModel = MakeShareable(new FStreamingLevelModel(Editor, *this, StreamingLevel));
+			TSharedPtr<FStreamingLevelModel> LevelModel = MakeShareable(new FStreamingLevelModel(*this, StreamingLevel));
 			AllLevelsList.Add(LevelModel);
 			AllLevelsMap.Add(LevelModel->GetLongPackageName(), LevelModel);
 
@@ -573,7 +573,7 @@ void FStreamingLevelCollectionModel::MergeSelectedLevels_Executed()
 			FEditorDelegates::NewCurrentLevel.Broadcast();
 		}
 
-		Editor->NoteSelectionChange();
+		GEditor->NoteSelectionChange();
 
 		//restore the original selection and remove the levels that were merged
 		SetSelectedLevels(SelectedLevelsCopy);
@@ -655,16 +655,16 @@ void FStreamingLevelCollectionModel::SelectStreamingVolumes_Executed()
 
 	// Select the volumes.
 	const FScopedTransaction Transaction( LOCTEXT("SelectAssociatedStreamingVolumes", "Select Associated Streaming Volumes") );
-	Editor->GetSelectedActors()->Modify();
-	Editor->SelectNone(false, true);
+	GEditor->GetSelectedActors()->Modify();
+	GEditor->SelectNone(false, true);
 
 	for (int32 i = 0 ; i < LevelStreamingVolumesToSelect.Num() ; ++i)
 	{
 		ALevelStreamingVolume* LevelStreamingVolume = LevelStreamingVolumesToSelect[i];
-		Editor->SelectActor(LevelStreamingVolume, /*bInSelected=*/ true, false, true);
+		GEditor->SelectActor(LevelStreamingVolume, /*bInSelected=*/ true, false, true);
 	}
 	
-	Editor->NoteSelectionChange();
+	GEditor->NoteSelectionChange();
 }
 
 #undef LOCTEXT_NAMESPACE
diff --git a/Engine/Source/Editor/WorldBrowser/Private/StreamingLevels/StreamingLevelCollectionModel.h b/Engine/Source/Editor/WorldBrowser/Private/StreamingLevels/StreamingLevelCollectionModel.h
--- a/Engine/Source/Editor/WorldBrowser/Private/StreamingLevels/StreamingLevelCollectionModel.h
+++ b/Engine/Source/Editor/WorldBrowser/Private/StreamingLevels/StreamingLevelCollectionModel.h
@@ -20,9 +20,9 @@ public:
 	 *
 	 *	@param	InEditor		The UEditorEngine to use
 	 */
-	static TSharedRef<FStreamingLevelCollectionModel> Create(UEditorEngine* InEditor, UWorld* InWorld)
+	static TSharedRef<FStreamingLevelCollectionModel> Create(UWorld* InWorld)
 	{
-		TSharedRef<FStreamingLevelCollectionModel> LevelCollectionModel(new FStreamingLevelCollectionModel(InEditor));
+		TSharedRef<FStreamingLevelCollectionModel> LevelCollectionModel(new FStreamingLevelCollectionModel());
 		LevelCollectionModel->Initialize(InWorld);
 		return LevelCollectionModel;
 	}
@@ -55,7 +55,7 @@ private:
 	 *	@param	InWorldLevels	The Level management logic object
 	 *	@param	InEditor		The UEditorEngine to use
 	 */
-	FStreamingLevelCollectionModel(UEditorEngine* InEditor);
+	FStreamingLevelCollectionModel();
 	
 	// Begin FEditorUndoClient Interface
 	virtual void PostUndo(bool bSuccess) override { UpdateAllLevels(); }
diff --git a/Engine/Source/Editor/WorldBrowser/Private/StreamingLevels/StreamingLevelModel.cpp b/Engine/Source/Editor/WorldBrowser/Private/StreamingLevels/StreamingLevelModel.cpp
--- a/Engine/Source/Editor/WorldBrowser/Private/StreamingLevels/StreamingLevelModel.cpp
+++ b/Engine/Source/Editor/WorldBrowser/Private/StreamingLevels/StreamingLevelModel.cpp
@@ -9,20 +9,18 @@
 
 #define LOCTEXT_NAMESPACE "WorldBrowser"
 
-FStreamingLevelModel::FStreamingLevelModel(const TWeakObjectPtr<UEditorEngine>& InEditor, 
-											FStreamingLevelCollectionModel& InWorldData, 
-											class ULevelStreaming* InLevelStreaming)
+FStreamingLevelModel::FStreamingLevelModel(FStreamingLevelCollectionModel& InWorldData, ULevelStreaming* InLevelStreaming)
 
-	: FLevelModel(InWorldData, InEditor)
+	: FLevelModel(InWorldData)
 	, LevelStreaming(InLevelStreaming)
 	, bHasValidPackageName(false)
 {
-	Editor->RegisterForUndo( this );
+	GEditor->RegisterForUndo( this );
 }
 
 FStreamingLevelModel::~FStreamingLevelModel()
 {
-	Editor->UnregisterForUndo( this );
+	GEditor->UnregisterForUndo( this );
 }
 
 bool FStreamingLevelModel::HasValidPackage() const
diff --git a/Engine/Source/Editor/WorldBrowser/Private/StreamingLevels/StreamingLevelModel.h b/Engine/Source/Editor/WorldBrowser/Private/StreamingLevels/StreamingLevelModel.h
--- a/Engine/Source/Editor/WorldBrowser/Private/StreamingLevels/StreamingLevelModel.h
+++ b/Engine/Source/Editor/WorldBrowser/Private/StreamingLevels/StreamingLevelModel.h
@@ -21,9 +21,7 @@ public:
 	 *	@param	InWorldData			Level collection owning this model
 	 *	@param	InLevelStreaming	Streaming object this model should represent
 	 */
-	FStreamingLevelModel(const TWeakObjectPtr<UEditorEngine>& InEditor, 
-							FStreamingLevelCollectionModel& InWorldData, 
-							class ULevelStreaming* InLevelStreaming);
+	FStreamingLevelModel(FStreamingLevelCollectionModel& InWorldData, class ULevelStreaming* InLevelStreaming);
 	~FStreamingLevelModel();
 
 
diff --git a/Engine/Source/Editor/WorldBrowser/Private/Tiles/WorldTileCollectionModel.cpp b/Engine/Source/Editor/WorldBrowser/Private/Tiles/WorldTileCollectionModel.cpp
--- a/Engine/Source/Editor/WorldBrowser/Private/Tiles/WorldTileCollectionModel.cpp
+++ b/Engine/Source/Editor/WorldBrowser/Private/Tiles/WorldTileCollectionModel.cpp
@@ -32,8 +32,8 @@
 static const FName HeightmapLayerName = FName("__Heightmap__");
 
 
-FWorldTileCollectionModel::FWorldTileCollectionModel(UEditorEngine* InEditor)
-	: FLevelCollectionModel(InEditor)
+FWorldTileCollectionModel::FWorldTileCollectionModel()
+	: FLevelCollectionModel()
 	, PreviewLocation(0.0f,0.0f,0.0f)
 	, bIsSavingLevel(false)
 	, bMeshProxyAvailable(false)
@@ -47,7 +47,7 @@ FWorldTileCollectionModel::~FWorldTileCollectionModel()
 	
 	CurrentWorld = NULL;
 
-	Editor->UnregisterForUndo(this);
+	GEditor->UnregisterForUndo(this);
 	FCoreDelegates::PreWorldOriginOffset.RemoveAll(this);
 	FCoreDelegates::PostWorldOriginOffset.RemoveAll(this);
 	FEditorDelegates::PreSaveWorld.RemoveAll(this);
@@ -62,7 +62,7 @@ void FWorldTileCollectionModel::Initialize(UWorld* InWorld)
 	ManagedLayers.Empty();
 	ManagedLayers.Add(Layer);
 
-	Editor->RegisterForUndo(this);
+	GEditor->RegisterForUndo(this);
 	FCoreDelegates::PreWorldOriginOffset.AddSP(this, &FWorldTileCollectionModel::PreWorldOriginOffset);
 	FCoreDelegates::PostWorldOriginOffset.AddSP(this, &FWorldTileCollectionModel::PostWorldOriginOffset);
 	FEditorDelegates::PreSaveWorld.AddSP(this, &FWorldTileCollectionModel::OnPreSaveWorld);
@@ -595,10 +595,8 @@ bool FWorldTileCollectionModel::GetPlayerView(FVector& Location, FRotator& Rotat
 
 bool FWorldTileCollectionModel::GetObserverView(FVector& Location, FRotator& Rotation) const
 {
-	const UEditorEngine* EditorEngine = Editor.Get();
-
 	// We are in the SIE
-	if (EditorEngine->bIsSimulatingInEditor && GCurrentLevelEditingViewportClient->IsSimulateInEditorViewport())
+	if (GEditor->bIsSimulatingInEditor && GCurrentLevelEditingViewportClient->IsSimulateInEditorViewport())
 	{
 		Rotation = GCurrentLevelEditingViewportClient->GetViewRotation();
 		Location = GCurrentLevelEditingViewportClient->GetViewLocation();
@@ -606,9 +604,9 @@ bool FWorldTileCollectionModel::GetObserverView(FVector& Location, FRotator& Rot
 	}
 
 	// We are in the editor world
-	if (EditorEngine->PlayWorld == nullptr)
+	if (GEditor->PlayWorld == nullptr)
 	{
-		for (const FLevelEditorViewportClient* ViewportClient : EditorEngine->LevelViewportClients)
+		for (const FLevelEditorViewportClient* ViewportClient : GEditor->LevelViewportClients)
 		{
 			if (ViewportClient && ViewportClient->IsPerspective())
 			{
@@ -788,7 +786,7 @@ void FWorldTileCollectionModel::OnLevelsCollectionChanged()
 
 TSharedPtr<FWorldTileModel> FWorldTileCollectionModel::AddLevelFromTile(int32 TileIdx)
 {
-	TSharedPtr<FWorldTileModel> LevelModel = MakeShareable(new FWorldTileModel(Editor, *this, TileIdx));
+	TSharedPtr<FWorldTileModel> LevelModel = MakeShareable(new FWorldTileModel(*this, TileIdx));
 	AllLevelsList.Add(LevelModel);
 	AllLevelsMap.Add(LevelModel->TileDetails->PackageName, LevelModel);
 	
@@ -1542,7 +1540,7 @@ void FWorldTileCollectionModel::ImportTiledLandscape_Executed()
 	ImportWidnow->SetContent(ImportDialog);
 
 	/** Show the dialog window as a modal window */
-	Editor->EditorAddModalWindow(ImportWidnow);
+	GEditor->EditorAddModalWindow(ImportWidnow);
 		
 	if (ImportDialog->ShouldImport() && ImportDialog->GetImportSettings().HeightmapFileList.Num())
 	{
@@ -2132,7 +2130,7 @@ bool FWorldTileCollectionModel::GenerateLODLevels(FLevelModelList InLevelList, i
 				for (UObject* Asset : GeneratedAssets)
 				{
 					FAssetRegistryModule::AssetCreated(Asset);
-					Editor->BroadcastObjectReimported(Asset);
+					GEditor->BroadcastObjectReimported(Asset);
 					PackagesToSave.Add(Asset->GetOutermost());
 				}
 								
diff --git a/Engine/Source/Editor/WorldBrowser/Private/Tiles/WorldTileCollectionModel.h b/Engine/Source/Editor/WorldBrowser/Private/Tiles/WorldTileCollectionModel.h
--- a/Engine/Source/Editor/WorldBrowser/Private/Tiles/WorldTileCollectionModel.h
+++ b/Engine/Source/Editor/WorldBrowser/Private/Tiles/WorldTileCollectionModel.h
@@ -19,9 +19,9 @@ public:
 	 *
 	 *	@param	InEditor		The UEditorEngine to use
 	 */
-	static TSharedRef<FWorldTileCollectionModel> Create(UEditorEngine* InEditor, UWorld* InWorld)
+	static TSharedRef<FWorldTileCollectionModel> Create(UWorld* InWorld)
 	{
-		TSharedRef<FWorldTileCollectionModel> LevelCollectionModel(new FWorldTileCollectionModel(InEditor));
+		TSharedRef<FWorldTileCollectionModel> LevelCollectionModel(new FWorldTileCollectionModel());
 		LevelCollectionModel->Initialize(InWorld);
 		return LevelCollectionModel;
 	}
@@ -145,7 +145,7 @@ public:
 	void BuildWorldCompositionMenu(FMenuBuilder& InMenuBuilder) const;
 
 private:
-	FWorldTileCollectionModel(UEditorEngine* InEditor);
+	FWorldTileCollectionModel();
 	
 	/** Setups parent->child links between tiles */
 	void SetupParentChildLinks();
diff --git a/Engine/Source/Editor/WorldBrowser/Private/Tiles/WorldTileModel.cpp b/Engine/Source/Editor/WorldBrowser/Private/Tiles/WorldTileModel.cpp
--- a/Engine/Source/Editor/WorldBrowser/Private/Tiles/WorldTileModel.cpp
+++ b/Engine/Source/Editor/WorldBrowser/Private/Tiles/WorldTileModel.cpp
@@ -19,10 +19,8 @@
 #define LOCTEXT_NAMESPACE "WorldBrowser"
 DEFINE_LOG_CATEGORY_STATIC(WorldBrowser, Log, All);
 
-FWorldTileModel::FWorldTileModel(const TWeakObjectPtr<UEditorEngine>& InEditor, 
-								 FWorldTileCollectionModel& InWorldModel, 
-								 int32 InTileIdx)
-	: FLevelModel(InWorldModel, InEditor) 
+FWorldTileModel::FWorldTileModel(FWorldTileCollectionModel& InWorldModel, int32 InTileIdx)
+	: FLevelModel(InWorldModel) 
 	, TileIdx(InTileIdx)
 	, TileDetails(NULL)
 	, bWasShelved(false)
@@ -469,7 +467,7 @@ void FWorldTileModel::SetLevelPosition(const FIntPoint& InPosition)
 			{
 				if (Actor != nullptr)
 				{
-					Editor->BroadcastOnActorMoved(Actor);
+					GEditor->BroadcastOnActorMoved(Actor);
 				}
 			}
 		}
diff --git a/Engine/Source/Editor/WorldBrowser/Private/Tiles/WorldTileModel.h b/Engine/Source/Editor/WorldBrowser/Private/Tiles/WorldTileModel.h
--- a/Engine/Source/Editor/WorldBrowser/Private/Tiles/WorldTileModel.h
+++ b/Engine/Source/Editor/WorldBrowser/Private/Tiles/WorldTileModel.h
@@ -58,12 +58,10 @@ public:
 	/**
 	 *	FWorldTileModel Constructor
 	 *
-	 *	@param	InEditor		The UEditorEngine to use
 	 *	@param	InWorldData		WorldBrowser world data
 	 *	@param	TileIdx			Tile index in world composition tiles list
 	 */
-	FWorldTileModel(const TWeakObjectPtr<UEditorEngine>& InEditor, 
-					FWorldTileCollectionModel& InWorldData, int32 TileIdx);
+	FWorldTileModel(FWorldTileCollectionModel& InWorldData, int32 TileIdx);
 	~FWorldTileModel();
 
 public:
diff --git a/Engine/Source/Editor/WorldBrowser/Private/WorldBrowserModule.cpp b/Engine/Source/Editor/WorldBrowser/Private/WorldBrowserModule.cpp
--- a/Engine/Source/Editor/WorldBrowser/Private/WorldBrowserModule.cpp
+++ b/Engine/Source/Editor/WorldBrowser/Private/WorldBrowserModule.cpp
@@ -113,11 +113,11 @@ TSharedPtr<FLevelCollectionModel> FWorldBrowserModule::SharedWorldModel(UWorld*
 		{
 			if (InWorld->WorldComposition)
 			{
-				SharedWorldModel = FWorldTileCollectionModel::Create(GEditor, InWorld);
+				SharedWorldModel = FWorldTileCollectionModel::Create(InWorld);
 			}
 			else
 			{
-				SharedWorldModel = FStreamingLevelCollectionModel::Create(GEditor, InWorld);
+				SharedWorldModel = FStreamingLevelCollectionModel::Create(InWorld);
 			}
 		}
 
