Added level coloration support for world compostion sub-levels
UE-16046

diff --git a/Engine/Source/Editor/WorldBrowser/Private/LevelModel.cpp b/Engine/Source/Editor/WorldBrowser/Private/LevelModel.cpp
--- a/Engine/Source/Editor/WorldBrowser/Private/LevelModel.cpp
+++ b/Engine/Source/Editor/WorldBrowser/Private/LevelModel.cpp
@@ -354,12 +354,6 @@ void FLevelModel::SetLevelTranslationDelta(FVector2D InAbsoluteDelta)
 	}
 }
 
-bool FLevelModel::SupportsLevelColor() const
-{
-	// By default Levels don't support a level color
-	return false;
-}
-
 FLinearColor FLevelModel::GetLevelColor() const
 {
 	// Returns Constant color, base classes will override this
diff --git a/Engine/Source/Editor/WorldBrowser/Private/LevelModel.h b/Engine/Source/Editor/WorldBrowser/Private/LevelModel.h
--- a/Engine/Source/Editor/WorldBrowser/Private/LevelModel.h
+++ b/Engine/Source/Editor/WorldBrowser/Private/LevelModel.h
@@ -141,9 +141,6 @@ public:
 	/** Sets new translation delta to this model and all descendants*/
 	void SetLevelTranslationDelta(FVector2D InAbsoluteDelta);
 
-	/** @return Whether level supports a Level Color for visualization. (Show -> Advanced -> Level Coloration) */
-	virtual bool SupportsLevelColor() const;
-
 	/** @return level color, used for visualization. (Show -> Advanced -> Level Coloration) */
 	virtual FLinearColor GetLevelColor() const;
 
diff --git a/Engine/Source/Editor/WorldBrowser/Private/SWorldHierarchyItem.cpp b/Engine/Source/Editor/WorldBrowser/Private/SWorldHierarchyItem.cpp
--- a/Engine/Source/Editor/WorldBrowser/Private/SWorldHierarchyItem.cpp
+++ b/Engine/Source/Editor/WorldBrowser/Private/SWorldHierarchyItem.cpp
@@ -111,30 +111,23 @@ TSharedRef< SWidget > SWorldHierarchyItem::GenerateWidgetForColumn( const FName&
 	}
 	else if (ColumnID == HierarchyColumns::ColumnID_Color)
 	{
-		if (LevelModel->SupportsLevelColor())
-		{
-			TableRowContent =
-				SAssignNew(ColorButton, SButton)
-				.ContentPadding(0)
-				.ButtonStyle(FEditorStyle::Get(), "ToggleButton")
-				.IsEnabled(true)
-				.OnClicked(this, &SWorldHierarchyItem::OnChangeColor)
-				.ToolTipText(LOCTEXT("LevelColorButtonToolTip", "Change Level Color"))
-				.HAlign(HAlign_Center)
-				.VAlign(VAlign_Center)
-				.Content()
-				[
-					SNew(SImage)
-					.ColorAndOpacity(this, &SWorldHierarchyItem::GetDrawColor)
-					.Image(this, &SWorldHierarchyItem::GetLevelColorBrush)
-				]
-			;
-		}
-		else
-		{
-			TableRowContent =
-				SNew(SHorizontalBox);
-		}
+		TableRowContent =
+			SAssignNew(ColorButton, SButton)
+			.ContentPadding(0)
+			.ButtonStyle(FEditorStyle::Get(), "ToggleButton")
+			.IsEnabled(true)
+			.OnClicked(this, &SWorldHierarchyItem::OnChangeColor)
+			.ToolTipText(LOCTEXT("LevelColorButtonToolTip", "Change Level Color"))
+			.HAlign(HAlign_Center)
+			.VAlign(VAlign_Center)
+			.Visibility(this, &SWorldHierarchyItem::GetColorButtonVisibility)
+			.Content()
+			[
+				SNew(SImage)
+				.ColorAndOpacity(this, &SWorldHierarchyItem::GetDrawColor)
+				.Image(this, &SWorldHierarchyItem::GetLevelColorBrush)
+			]
+		;
 	}
 	else if (ColumnID == HierarchyColumns::ColumnID_Kismet)
 	{
@@ -332,6 +325,20 @@ FReply SWorldHierarchyItem::OnChangeColor()
 	return FReply::Handled();
 }
 
+EVisibility SWorldHierarchyItem::GetColorButtonVisibility() const
+{
+	EVisibility Result = EVisibility::Hidden;
+	if (LevelModel.IsValid())
+	{
+		ULevel* LevelObject = LevelModel->GetLevelObject();
+		if (LevelObject && !LevelObject->IsPersistentLevel())
+		{
+			Result = EVisibility::Visible;
+		}
+	}
+	return Result;
+}
+
 FReply SWorldHierarchyItem::OnItemDragDetected(const FGeometry& MyGeometry, const FPointerEvent& MouseEvent)
 {
 	if (MouseEvent.IsMouseButtonDown(EKeys::LeftMouseButton))
diff --git a/Engine/Source/Editor/WorldBrowser/Private/SWorldHierarchyItem.h b/Engine/Source/Editor/WorldBrowser/Private/SWorldHierarchyItem.h
--- a/Engine/Source/Editor/WorldBrowser/Private/SWorldHierarchyItem.h
+++ b/Engine/Source/Editor/WorldBrowser/Private/SWorldHierarchyItem.h
@@ -100,6 +100,9 @@ private:
 	void OnColorPickerCancelled(FLinearColor OriginalColor);
 	void OnColorPickerInteractiveBegin();
 	void OnColorPickerInteractiveEnd();
+
+	/** Whether color button should be visible, depends on whether sub-level is loaded */
+	EVisibility GetColorButtonVisibility() const;
 	
 	/**
 	 *  @return The text of level name while it is not being edited
diff --git a/Engine/Source/Editor/WorldBrowser/Private/StreamingLevels/StreamingLevelModel.cpp b/Engine/Source/Editor/WorldBrowser/Private/StreamingLevels/StreamingLevelModel.cpp
--- a/Engine/Source/Editor/WorldBrowser/Private/StreamingLevels/StreamingLevelModel.cpp
+++ b/Engine/Source/Editor/WorldBrowser/Private/StreamingLevels/StreamingLevelModel.cpp
@@ -72,11 +72,6 @@ void FStreamingLevelModel::UpdateAsset(const FAssetData& AssetData)
 	}
 }
 
-bool FStreamingLevelModel::SupportsLevelColor() const
-{
-	return LevelStreaming.IsValid();
-}
-
 FLinearColor FStreamingLevelModel::GetLevelColor() const
 {
 	if (LevelStreaming.IsValid())
diff --git a/Engine/Source/Editor/WorldBrowser/Private/StreamingLevels/StreamingLevelModel.h b/Engine/Source/Editor/WorldBrowser/Private/StreamingLevels/StreamingLevelModel.h
--- a/Engine/Source/Editor/WorldBrowser/Private/StreamingLevels/StreamingLevelModel.h
+++ b/Engine/Source/Editor/WorldBrowser/Private/StreamingLevels/StreamingLevelModel.h
@@ -33,7 +33,6 @@ public:
 	virtual FName GetAssetName() const override;
 	virtual FName GetLongPackageName() const override;
 	virtual void UpdateAsset(const FAssetData& AssetData) override;
-	virtual bool SupportsLevelColor() const override;
 	virtual FLinearColor GetLevelColor() const override;
 	virtual void SetLevelColor(FLinearColor InColor) override;
 	virtual void Update() override;
diff --git a/Engine/Source/Editor/WorldBrowser/Private/Tiles/WorldTileModel.cpp b/Engine/Source/Editor/WorldBrowser/Private/Tiles/WorldTileModel.cpp
--- a/Engine/Source/Editor/WorldBrowser/Private/Tiles/WorldTileModel.cpp
+++ b/Engine/Source/Editor/WorldBrowser/Private/Tiles/WorldTileModel.cpp
@@ -581,6 +581,14 @@ void FWorldTileModel::LoadLevel()
 	bWasShelved = !bShouldBeVisible;
 	//
 	LoadedLevel = LevelStreaming->GetLoadedLevel();
+	
+	if (LoadedLevel.IsValid())
+	{
+		// LevelStreaming is transient object, but components fetch color from it
+		LevelStreaming->LevelColor = LoadedLevel->LevelColor;
+		LoadedLevel->MarkLevelComponentsRenderStateDirty();
+	}
+	
 	// Enable tile properties
 	TileDetails->bTileEditable = (LoadedLevel != nullptr);
 }
@@ -605,7 +613,7 @@ ULevelStreaming* FWorldTileModel::GetAssosiatedStreamingLevel()
 
 		//
 		AssociatedStreamingLevel->SetWorldAssetByPackageName(PackageName);
-		AssociatedStreamingLevel->LevelColor		= FLinearColor::MakeRandomColor();
+		AssociatedStreamingLevel->LevelColor		= GetLevelColor();
 		AssociatedStreamingLevel->LevelTransform	= FTransform::Identity;
 		AssociatedStreamingLevel->PackageNameToLoad	= PackageName;
 		//
@@ -660,6 +668,35 @@ bool FWorldTileModel::IsVisibleInCompositionView() const
 	return !TileDetails->bHideInTileView && LevelCollectionModel.PassesAllFilters(*this);
 }
 
+FLinearColor FWorldTileModel::GetLevelColor() const
+{
+	ULevel* LevelObject = GetLevelObject();
+	if (LevelObject)
+	{
+		return LevelObject->LevelColor;
+	}
+	else
+	{
+		return FLevelModel::GetLevelColor();
+	}
+}
+
+void FWorldTileModel::SetLevelColor(FLinearColor InColor)
+{
+	ULevel* LevelObject = GetLevelObject();
+	if (LevelObject)
+	{
+		ULevelStreaming* StreamingLevel = GetAssosiatedStreamingLevel();
+		if (StreamingLevel)
+		{
+			LevelObject->MarkPackageDirty();
+			LevelObject->LevelColor = InColor;
+			StreamingLevel->LevelColor = InColor; // this is transient object, but components fetch color from it
+			LevelObject->MarkLevelComponentsRenderStateDirty();
+		}
+	}
+}
+
 void FWorldTileModel::OnLevelBoundsActorUpdated()
 {
 	Update();
diff --git a/Engine/Source/Editor/WorldBrowser/Private/Tiles/WorldTileModel.h b/Engine/Source/Editor/WorldBrowser/Private/Tiles/WorldTileModel.h
--- a/Engine/Source/Editor/WorldBrowser/Private/Tiles/WorldTileModel.h
+++ b/Engine/Source/Editor/WorldBrowser/Private/Tiles/WorldTileModel.h
@@ -82,6 +82,8 @@ public:
 	virtual void OnLevelRemovedFromWorld() override;
 	virtual void OnParentChanged() override;
 	virtual bool IsVisibleInCompositionView() const override;
+	virtual FLinearColor GetLevelColor() const override;
+	virtual void SetLevelColor(FLinearColor InColor) override;
 	// FLevelModel interface end
 	
 	/** Adds new streaming level*/
diff --git a/Engine/Source/Runtime/Engine/Classes/Engine/Level.h b/Engine/Source/Runtime/Engine/Classes/Engine/Level.h
--- a/Engine/Source/Runtime/Engine/Classes/Engine/Level.h
+++ b/Engine/Source/Runtime/Engine/Classes/Engine/Level.h
@@ -468,6 +468,13 @@ public:
 	/** Level simplification settings for each LOD */
 	UPROPERTY()
 	FLevelSimplificationDetails LevelSimplification[WORLDTILE_LOD_MAX_INDEX];
+
+	/** 
+	 * The level color used for visualization. (Show -> Advanced -> Level Coloration)
+	 * Used only in world composition mode
+	 */
+	UPROPERTY()
+	FLinearColor LevelColor;
 #endif //WITH_EDITORONLY_DATA
 
 	/** Actor which defines level logical bounding box				*/
@@ -582,8 +589,11 @@ public:
 	 * Invalidates the cached data used to render the level's UModel.
 	 */
 	void InvalidateModelGeometry();
-
+	
 #if WITH_EDITOR
+	/** Marks all level components render state as dirty */
+	ENGINE_API void MarkLevelComponentsRenderStateDirty();
+
 	/** Called to create ModelComponents for BSP rendering */
 	void CreateModelComponents();
 #endif // WITH_EDITOR
diff --git a/Engine/Source/Runtime/Engine/Private/Level.cpp b/Engine/Source/Runtime/Engine/Private/Level.cpp
--- a/Engine/Source/Runtime/Engine/Private/Level.cpp
+++ b/Engine/Source/Runtime/Engine/Private/Level.cpp
@@ -237,6 +237,9 @@ ULevel::ULevel( const FObjectInitializer& ObjectInitializer )
 	,	TickTaskLevel(FTickTaskManagerInterface::Get().AllocateTickTaskLevel())
 	,	PrecomputedLightVolume(new FPrecomputedLightVolume())
 {
+#if WITH_EDITORONLY_DATA
+	LevelColor = FLinearColor::White;
+#endif
 }
 
 void ULevel::Initialize(const FURL& InURL)
@@ -802,6 +805,25 @@ void ULevel::IncrementalUpdateComponents(int32 NumComponentsToUpdate, bool bReru
 
 #if WITH_EDITOR
 
+void ULevel::MarkLevelComponentsRenderStateDirty()
+{
+	for (UModelComponent* ModelComponent : ModelComponents)
+	{
+		if (ModelComponent)
+		{
+			ModelComponent->MarkRenderStateDirty();
+		}
+	}
+
+	for (AActor* Actor : Actors)
+	{
+		if (Actor)
+		{
+			Actor->MarkComponentsRenderStateDirty();
+		}
+	}
+}
+
 void ULevel::CreateModelComponents()
 {
 	// Update the model vertices and edges.
diff --git a/Engine/Source/Runtime/Engine/Private/LevelStreaming.cpp b/Engine/Source/Runtime/Engine/Private/LevelStreaming.cpp
--- a/Engine/Source/Runtime/Engine/Private/LevelStreaming.cpp
+++ b/Engine/Source/Runtime/Engine/Private/LevelStreaming.cpp
@@ -870,20 +870,9 @@ void ULevelStreaming::PostEditChangeProperty(FPropertyChangedEvent& PropertyChan
 		{
 			// Make sure the level's Level Color change is applied immediately by reregistering the
 			// components of the actor's in the level
-			if( LoadedLevel != NULL )
+			if (LoadedLevel != nullptr)
 			{
-				UPackage* Package = LoadedLevel->GetOutermost();
-				for( TObjectIterator<UActorComponent> It; It; ++It )
-				{
-					if( It->IsIn( Package ) )
-					{
-						UActorComponent* ActorComponent = Cast<UActorComponent>( *It );
-						if( ActorComponent )
-						{
-							ActorComponent->RecreateRenderState_Concurrent();
-						}
-					}
-				}
+				LoadedLevel->MarkLevelComponentsRenderStateDirty();
 			}
 		}
 	}
