Do not use level streaming for loading sub-levels in world composition editor

diff --git a/Engine/Source/Editor/WorldBrowser/Private/Tiles/WorldTileModel.cpp b/Engine/Source/Editor/WorldBrowser/Private/Tiles/WorldTileModel.cpp
--- a/Engine/Source/Editor/WorldBrowser/Private/Tiles/WorldTileModel.cpp
+++ b/Engine/Source/Editor/WorldBrowser/Private/Tiles/WorldTileModel.cpp
@@ -566,29 +566,55 @@ void FWorldTileModel::LoadLevel()
 
 	// Create transient level streaming object and add to persistent level
 	ULevelStreaming* LevelStreaming = GetAssosiatedStreamingLevel();
+	// should be clean level streaming object here
+	check(LevelStreaming && LevelStreaming->GetLoadedLevel() == nullptr);
+	
+	bLoadingLevel = true;
+
+	// Load level package 
+	{
+		FName LevelPackageName = LevelStreaming->GetWorldAssetPackageFName();
+		
+		ULevel::StreamedLevelsOwningWorld.Add(LevelPackageName, LevelCollectionModel.GetWorld());
+		UWorld::WorldTypePreLoadMap.FindOrAdd(LevelPackageName) = LevelCollectionModel.GetWorld()->WorldType;
+
+		UPackage* LevelPackage = LoadPackage(nullptr, *LevelPackageName.ToString(), LOAD_None);
 
+		ULevel::StreamedLevelsOwningWorld.Remove(LevelPackageName);
+		UWorld::WorldTypePreLoadMap.Remove(LevelPackageName);
+
+		// Find world object and use its PersistentLevel pointer.
+		UWorld* LevelWorld = UWorld::FindWorldInPackage(LevelPackage);
+		// Check for a redirector. Follow it, if found.
+		if (LevelWorld == nullptr)
+		{
+			LevelWorld = UWorld::FollowWorldRedirectorInPackage(LevelPackage);
+		}
+
+		if (LevelWorld && LevelWorld->PersistentLevel)
+		{
+			// LevelStreaming is transient object so world composition stores color in ULevel object
+			LevelStreaming->LevelColor = LevelWorld->PersistentLevel->LevelColor;
+		}
+	}
+	
 	// Whether this tile should be made visible at current world bounds
 	const bool bShouldBeVisible = ShouldBeVisible(LevelCollectionModel.EditableWorldArea());
 	
-	// Load level
-	bLoadingLevel = true;
+	// Our level package should be loaded at this point, so level streaming will find it in memory
 	LevelStreaming->bShouldBeLoaded = true;
 	LevelStreaming->bShouldBeVisible = false; // Should be always false in the Editor
-	LevelStreaming->bShouldBeVisibleInEditor = bShouldBeVisible;
+	LevelStreaming->bShouldBeVisibleInEditor = bShouldBeVisible; 
+	// Bring level to world
 	LevelCollectionModel.GetWorld()->FlushLevelStreaming();
+	
 	bLoadingLevel = false;
+
 	// Mark tile as shelved in case it is hidden(does not fit to world bounds)
 	bWasShelved = !bShouldBeVisible;
 	//
 	LoadedLevel = LevelStreaming->GetLoadedLevel();
 	
-	if (LoadedLevel.IsValid())
-	{
-		// LevelStreaming is transient object, but components fetch color from it
-		LevelStreaming->LevelColor = LoadedLevel->LevelColor;
-		LoadedLevel->MarkLevelComponentsRenderStateDirty();
-	}
-	
 	// Enable tile properties
 	TileDetails->bTileEditable = (LoadedLevel != nullptr);
 }
