Undo support added for level locking. Further work is required to support undoing read-only locking.
[UE-20836] Re-implement undo/redo for level locking/unlocking

diff --git a/Engine/Source/Editor/UnrealEd/Private/EditorLevelUtils.cpp b/Engine/Source/Editor/UnrealEd/Private/EditorLevelUtils.cpp
--- a/Engine/Source/Editor/UnrealEd/Private/EditorLevelUtils.cpp
+++ b/Engine/Source/Editor/UnrealEd/Private/EditorLevelUtils.cpp
@@ -697,7 +697,7 @@ namespace EditorLevelUtils
 				StreamingLevel = FLevelUtils::FindStreamingLevel( Level );
 			}
 
-			//create a transaction so we can undo the visibilty toggle
+			// Create a transaction so we can undo the visibility toggle
 			const FScopedTransaction Transaction( LOCTEXT( "ToggleLevelVisibility", "Toggle Level Visibility" ) );		
 
 			// Handle the case of a streaming level
diff --git a/Engine/Source/Editor/WorldBrowser/Private/LevelCollectionModel.cpp b/Engine/Source/Editor/WorldBrowser/Private/LevelCollectionModel.cpp
--- a/Engine/Source/Editor/WorldBrowser/Private/LevelCollectionModel.cpp
+++ b/Engine/Source/Editor/WorldBrowser/Private/LevelCollectionModel.cpp
@@ -143,7 +143,7 @@ void FLevelCollectionModel::BindCommands()
 		FCanExecuteAction::CreateSP( this, &FLevelCollectionModel::AreAnySelectedLevelsEditable ) );
 	
 	ActionList.MapAction( Commands.World_UnockSelectedLevels,
-		FExecuteAction::CreateSP( this, &FLevelCollectionModel::UnockSelectedLevels_Executed  ),
+		FExecuteAction::CreateSP( this, &FLevelCollectionModel::UnlockSelectedLevels_Executed  ),
 		FCanExecuteAction::CreateSP( this, &FLevelCollectionModel::AreAnySelectedLevelsEditable ) );
 	
 	ActionList.MapAction( Commands.World_LockAllLevels,
@@ -442,27 +442,35 @@ void FLevelCollectionModel::ShowLevels(const FLevelModelList& InLevelList)
 
 void FLevelCollectionModel::UnlockLevels(const FLevelModelList& InLevelList)
 {
-	if (IsReadOnly())
-	{
-		return;
-	}
-	
-	for (auto It = InLevelList.CreateConstIterator(); It; ++It)
+	if (!IsReadOnly())
 	{
-		(*It)->SetLocked(false);
+		const FText UndoTransactionText = (InLevelList.Num() == 1) ?
+			LOCTEXT("UnlockLevel", "Unlock Level") :
+			LOCTEXT("UnlockMultipleLevels", "Unlock Multiple Levels");
+
+		const FScopedTransaction Transaction(UndoTransactionText);
+
+		for (auto It = InLevelList.CreateConstIterator(); It; ++It)
+		{
+			(*It)->SetLocked(false);
+		}
 	}
 }
 	
 void FLevelCollectionModel::LockLevels(const FLevelModelList& InLevelList)
 {
-	if (IsReadOnly())
-	{
-		return;
-	}
-	
-	for (auto It = InLevelList.CreateConstIterator(); It; ++It)
+	if (!IsReadOnly())
 	{
-		(*It)->SetLocked(true);
+		const FText UndoTransactionText = (InLevelList.Num() == 1) ?
+			LOCTEXT("LockLevel", "Lock Level") :
+			LOCTEXT("LockMultipleLevels", "Lock Multiple Levels");
+
+		const FScopedTransaction Transaction(UndoTransactionText);
+
+		for (auto It = InLevelList.CreateConstIterator(); It; ++It)
+		{
+			(*It)->SetLocked(true);
+		}
 	}
 }
 
@@ -1454,19 +1462,27 @@ void FLevelCollectionModel::LockSelectedLevels_Executed()
 	LockLevels(GetSelectedLevels());
 }
 
-void FLevelCollectionModel::UnockSelectedLevels_Executed()
+void FLevelCollectionModel::UnlockSelectedLevels_Executed()
 {
 	UnlockLevels(GetSelectedLevels());
 }
 
 void FLevelCollectionModel::LockAllLevels_Executed()
 {
-	LockLevels(GetFilteredLevels());
+	if (!IsReadOnly())
+	{
+		const FScopedTransaction Transaction(LOCTEXT("LockAllLevels", "Lock All Levels"));
+		LockLevels(GetFilteredLevels());
+	}
 }
 
 void FLevelCollectionModel::UnockAllLevels_Executed()
 {
-	UnlockLevels(GetFilteredLevels());
+	if (!IsReadOnly())
+	{
+		const FScopedTransaction Transaction(LOCTEXT("UnlockAllLevels", "Unlock All Levels"));
+		UnlockLevels(GetFilteredLevels());
+	}
 }
 
 void FLevelCollectionModel::ToggleReadOnlyLevels_Executed()
diff --git a/Engine/Source/Editor/WorldBrowser/Private/LevelCollectionModel.h b/Engine/Source/Editor/WorldBrowser/Private/LevelCollectionModel.h
--- a/Engine/Source/Editor/WorldBrowser/Private/LevelCollectionModel.h
+++ b/Engine/Source/Editor/WorldBrowser/Private/LevelCollectionModel.h
@@ -380,7 +380,7 @@ protected:
 	void LockSelectedLevels_Executed();
 
 	/** Unlocks selected levels */
-	void UnockSelectedLevels_Executed();
+	void UnlockSelectedLevels_Executed();
 
 	/** Locks all levels */
 	void LockAllLevels_Executed();
