Copying //UE4/Dev-Landscape to //UE4/Dev-Main (Source: //UE4/Dev-Landscape @ 3053591)

diff --git a/Engine/Source/Editor/LandscapeEditor/Private/LandscapeFileFormatPng.cpp b/Engine/Source/Editor/LandscapeEditor/Private/LandscapeFileFormatPng.cpp
new file mode 100644
--- /dev/null
+++ b/Engine/Source/Editor/LandscapeEditor/Private/LandscapeFileFormatPng.cpp
@@ -0,0 +1,250 @@
+// Copyright 1998-2016 Epic Games, Inc. All Rights Reserved.
+
+#include "LandscapeEditorPrivatePCH.h"
+#include "LandscapeFileFormatPng.h"
+
+#include "ImageWrapper.h"
+#include "Algo/Transform.h"
+
+#define LOCTEXT_NAMESPACE "LandscapeEditor.NewLandscape"
+
+FLandscapeHeightmapFileFormat_Png::FLandscapeHeightmapFileFormat_Png()
+{
+	FileTypeInfo.Description = LOCTEXT("FileFormatPng_HeightmapDesc", "Heightmap .png files");
+	FileTypeInfo.Extensions.Add(".png");
+	FileTypeInfo.bSupportsExport = true;
+}
+
+FLandscapeHeightmapInfo FLandscapeHeightmapFileFormat_Png::Validate(const TCHAR* HeightmapFilename) const
+{
+	FLandscapeHeightmapInfo Result;
+
+	TArray<uint8> ImportData;
+	if (!FFileHelper::LoadFileToArray(ImportData, HeightmapFilename, FILEREAD_Silent))
+	{
+		Result.ResultCode = ELandscapeImportResult::Error;
+		Result.ErrorMessage = LOCTEXT("Import_HeightmapFileReadError", "Error reading heightmap file");
+	}
+	else
+	{
+		IImageWrapperModule& ImageWrapperModule = FModuleManager::LoadModuleChecked<IImageWrapperModule>("ImageWrapper");
+		IImageWrapperPtr ImageWrapper = ImageWrapperModule.CreateImageWrapper(EImageFormat::PNG);
+
+		if (!ImageWrapper->SetCompressed(ImportData.GetData(), ImportData.Num()) || ImageWrapper->GetWidth() <= 0 || ImageWrapper->GetHeight() <= 0)
+		{
+			Result.ResultCode = ELandscapeImportResult::Error;
+			Result.ErrorMessage = LOCTEXT("Import_HeightmapFileCorruptPng", "The heightmap file cannot be read (corrupt png?)");
+		}
+		else
+		{
+			if (ImageWrapper->GetFormat() != ERGBFormat::Gray)
+			{
+				Result.ResultCode = ELandscapeImportResult::Warning;
+				Result.ErrorMessage = LOCTEXT("Import_HeightmapFileColorPng", "The heightmap file appears to be a color png, grayscale is expected. The import *can* continue, but the result may not be what you expect...");
+			}
+			else if (ImageWrapper->GetBitDepth() != 16)
+			{
+				Result.ResultCode = ELandscapeImportResult::Warning;
+				Result.ErrorMessage = LOCTEXT("Import_HeightmapFileLowBitDepth", "The heightmap file appears to be an 8-bit png, 16-bit is preferred. The import *can* continue, but the result may be lower quality than desired.");
+			}
+			FLandscapeFileResolution ImportResolution;
+			ImportResolution.Width = ImageWrapper->GetWidth();
+			ImportResolution.Height = ImageWrapper->GetHeight();
+			Result.PossibleResolutions.Add(ImportResolution);
+		}
+	}
+
+	// possible todo: support sCAL (XY scale) and pCAL (Z scale) png chunks for filling out Result.DataScale
+	// I don't know if any heightmap generation software uses these or not
+	// if we support their import we should make the exporter write them too
+
+	return Result;
+}
+
+FLandscapeHeightmapImportData FLandscapeHeightmapFileFormat_Png::Import(const TCHAR* HeightmapFilename, FLandscapeFileResolution ExpectedResolution) const
+{
+	FLandscapeHeightmapImportData Result;
+
+	TArray<uint8> TempData;
+	if (!FFileHelper::LoadFileToArray(TempData, HeightmapFilename, FILEREAD_Silent))
+	{
+		Result.ResultCode = ELandscapeImportResult::Error;
+		Result.ErrorMessage = LOCTEXT("Import_HeightmapFileReadError", "Error reading heightmap file");
+	}
+	else
+	{
+		IImageWrapperModule& ImageWrapperModule = FModuleManager::LoadModuleChecked<IImageWrapperModule>("ImageWrapper");
+		IImageWrapperPtr ImageWrapper = ImageWrapperModule.CreateImageWrapper(EImageFormat::PNG);
+
+		if (!ImageWrapper->SetCompressed(TempData.GetData(), TempData.Num()))
+		{
+			Result.ResultCode = ELandscapeImportResult::Error;
+			Result.ErrorMessage = LOCTEXT("Import_HeightmapFileCorruptPng", "The heightmap file cannot be read (corrupt png?)");
+		}
+		else if (ImageWrapper->GetWidth() != ExpectedResolution.Width || ImageWrapper->GetHeight() != ExpectedResolution.Height)
+		{
+			Result.ResultCode = ELandscapeImportResult::Error;
+			Result.ErrorMessage = LOCTEXT("Import_HeightmapResolutionMismatch", "The heightmap file's resolution does not match the requested resolution");
+		}
+		else
+		{
+			if (ImageWrapper->GetFormat() != ERGBFormat::Gray)
+			{
+				Result.ResultCode = ELandscapeImportResult::Warning;
+				Result.ErrorMessage = LOCTEXT("Import_HeightmapFileColorPng", "The heightmap file appears to be a color png, grayscale is expected. The import *can* continue, but the result may not be what you expect...");
+			}
+			else if (ImageWrapper->GetBitDepth() != 16)
+			{
+				Result.ResultCode = ELandscapeImportResult::Warning;
+				Result.ErrorMessage = LOCTEXT("Import_HeightmapFileLowBitDepth", "The heightmap file appears to be an 8-bit png, 16-bit is preferred. The import *can* continue, but the result may be lower quality than desired.");
+			}
+
+			const TArray<uint8>* RawData = nullptr;
+			if (ImageWrapper->GetBitDepth() <= 8)
+			{
+				if (!ImageWrapper->GetRaw(ERGBFormat::Gray, 8, RawData))
+				{
+					Result.ResultCode = ELandscapeImportResult::Error;
+					Result.ErrorMessage = LOCTEXT("Import_HeightmapFileCorruptPng", "The heightmap file cannot be read (corrupt png?)");
+				}
+				else
+				{
+					Result.Data.Empty(ExpectedResolution.Width * ExpectedResolution.Height);
+					Algo::Transform(*RawData, Result.Data, [](uint8 Value) { return Value * 0x101; }); // Expand to 16-bit
+				}
+			}
+			else
+			{
+				if (!ImageWrapper->GetRaw(ERGBFormat::Gray, 16, RawData))
+				{
+					Result.ResultCode = ELandscapeImportResult::Error;
+					Result.ErrorMessage = LOCTEXT("Import_HeightmapFileCorruptPng", "The heightmap file cannot be read (corrupt png?)");
+				}
+				else
+				{
+					Result.Data.Empty(ExpectedResolution.Width * ExpectedResolution.Height);
+					Result.Data.AddUninitialized(ExpectedResolution.Width * ExpectedResolution.Height);
+					FMemory::Memcpy(Result.Data.GetData(), RawData->GetData(), ExpectedResolution.Width * ExpectedResolution.Height * 2);
+				}
+			}
+		}
+	}
+
+	return Result;
+}
+
+void FLandscapeHeightmapFileFormat_Png::Export(const TCHAR* HeightmapFilename, TArrayView<const uint16> Data, FLandscapeFileResolution DataResolution, FVector Scale) const
+{
+	IImageWrapperModule& ImageWrapperModule = FModuleManager::LoadModuleChecked<IImageWrapperModule>("ImageWrapper");
+	IImageWrapperPtr ImageWrapper = ImageWrapperModule.CreateImageWrapper(EImageFormat::PNG);
+
+	if (ImageWrapper->SetRaw(Data.GetData(), Data.Num() * 2, DataResolution.Width, DataResolution.Height, ERGBFormat::Gray, 16))
+	{
+		const TArray<uint8>& TempData = ImageWrapper->GetCompressed();
+		FFileHelper::SaveArrayToFile(TempData, HeightmapFilename);
+	}
+}
+
+//////////////////////////////////////////////////////////////////////////
+
+FLandscapeWeightmapFileFormat_Png::FLandscapeWeightmapFileFormat_Png()
+{
+	FileTypeInfo.Description = LOCTEXT("FileFormatPng_WeightmapDesc", "Layer .png files");
+	FileTypeInfo.Extensions.Add(".png");
+	FileTypeInfo.bSupportsExport = true;
+}
+
+FLandscapeWeightmapInfo FLandscapeWeightmapFileFormat_Png::Validate(const TCHAR* WeightmapFilename, FName LayerName) const
+{
+	FLandscapeWeightmapInfo Result;
+
+	TArray<uint8> ImportData;
+	if (!FFileHelper::LoadFileToArray(ImportData, WeightmapFilename, FILEREAD_Silent))
+	{
+		Result.ResultCode = ELandscapeImportResult::Error;
+		Result.ErrorMessage = LOCTEXT("Import_LayerFileReadError", "Error reading layer file");
+	}
+	else
+	{
+		IImageWrapperModule& ImageWrapperModule = FModuleManager::LoadModuleChecked<IImageWrapperModule>("ImageWrapper");
+		IImageWrapperPtr ImageWrapper = ImageWrapperModule.CreateImageWrapper(EImageFormat::PNG);
+
+		if (!ImageWrapper->SetCompressed(ImportData.GetData(), ImportData.Num()))
+		{
+			Result.ResultCode = ELandscapeImportResult::Error;
+			Result.ErrorMessage = LOCTEXT("Import_LayerCorruptPng", "The layer file cannot be read (corrupt png?)");
+		}
+		else
+		{
+			if (ImageWrapper->GetFormat() != ERGBFormat::Gray)
+			{
+				Result.ResultCode = ELandscapeImportResult::Warning;
+				Result.ErrorMessage = LOCTEXT("Import_LayerColorPng", "The layer file appears to be a color png, grayscale is expected. The import *can* continue, but the result may not be what you expect...");
+			}
+			FLandscapeFileResolution ImportResolution;
+			ImportResolution.Width = ImageWrapper->GetWidth();
+			ImportResolution.Height = ImageWrapper->GetHeight();
+			Result.PossibleResolutions.Add(ImportResolution);
+		}
+	}
+
+	return Result;
+}
+
+FLandscapeWeightmapImportData FLandscapeWeightmapFileFormat_Png::Import(const TCHAR* WeightmapFilename, FName LayerName, FLandscapeFileResolution ExpectedResolution) const
+{
+	FLandscapeWeightmapImportData Result;
+
+	TArray<uint8> TempData;
+	if (!FFileHelper::LoadFileToArray(TempData, WeightmapFilename, FILEREAD_Silent))
+	{
+		Result.ResultCode = ELandscapeImportResult::Error;
+		Result.ErrorMessage = LOCTEXT("Import_LayerFileReadError", "Error reading layer file");
+	}
+	else
+	{
+		IImageWrapperModule& ImageWrapperModule = FModuleManager::LoadModuleChecked<IImageWrapperModule>("ImageWrapper");
+		IImageWrapperPtr ImageWrapper = ImageWrapperModule.CreateImageWrapper(EImageFormat::PNG);
+
+		const TArray<uint8>* RawData = nullptr;
+		if (!ImageWrapper->SetCompressed(TempData.GetData(), TempData.Num()))
+		{
+			Result.ResultCode = ELandscapeImportResult::Error;
+			Result.ErrorMessage = LOCTEXT("Import_LayerCorruptPng", "The layer file cannot be read (corrupt png?)");
+		}
+		else if (ImageWrapper->GetWidth() != ExpectedResolution.Width || ImageWrapper->GetHeight() != ExpectedResolution.Height)
+		{
+			Result.ResultCode = ELandscapeImportResult::Error;
+			Result.ErrorMessage = LOCTEXT("Import_LayerResolutionMismatch", "The layer file's resolution does not match the requested resolution");
+		}
+		else if (!ImageWrapper->GetRaw(ERGBFormat::Gray, 8, RawData))
+		{
+			Result.ResultCode = ELandscapeImportResult::Error;
+			Result.ErrorMessage = LOCTEXT("Import_LayerCorruptPng", "The layer file cannot be read (corrupt png?)");
+		}
+		else
+		{
+			if (ImageWrapper->GetFormat() != ERGBFormat::Gray)
+			{
+				Result.ResultCode = ELandscapeImportResult::Warning;
+				Result.ErrorMessage = LOCTEXT("Import_LayerColorPng", "The layer file appears to be a color png, grayscale is expected. The import *can* continue, but the result may not be what you expect...");
+			}
+
+			Result.Data = *RawData; // agh I want to use MoveTemp() here
+		}
+	}
+
+	return Result;
+}
+
+void FLandscapeWeightmapFileFormat_Png::Export(const TCHAR* WeightmapFilename, FName LayerName, TArrayView<const uint8> Data, FLandscapeFileResolution DataResolution) const
+{
+	IImageWrapperModule& ImageWrapperModule = FModuleManager::LoadModuleChecked<IImageWrapperModule>("ImageWrapper");
+	IImageWrapperPtr ImageWrapper = ImageWrapperModule.CreateImageWrapper(EImageFormat::PNG);
+
+	if (ImageWrapper->SetRaw(Data.GetData(), Data.Num(), DataResolution.Width, DataResolution.Height, ERGBFormat::Gray, 8))
+	{
+		const TArray<uint8>& TempData = ImageWrapper->GetCompressed();
+		FFileHelper::SaveArrayToFile(TempData, WeightmapFilename);
+	}
+}
diff --git a/Engine/Source/Editor/LandscapeEditor/Private/LandscapeFileFormatPng.h b/Engine/Source/Editor/LandscapeEditor/Private/LandscapeFileFormatPng.h
new file mode 100644
--- /dev/null
+++ b/Engine/Source/Editor/LandscapeEditor/Private/LandscapeFileFormatPng.h
@@ -0,0 +1,44 @@
+// Copyright 1998-2016 Epic Games, Inc. All Rights Reserved.
+
+#pragma once
+
+#include "LandscapeFileFormatInterface.h"
+
+// Implement .png file format
+class FLandscapeHeightmapFileFormat_Png : public ILandscapeHeightmapFileFormat
+{
+private:
+	FLandscapeFileTypeInfo FileTypeInfo;
+
+public:
+	FLandscapeHeightmapFileFormat_Png();
+
+	virtual const FLandscapeFileTypeInfo& GetInfo() const override
+	{
+		return FileTypeInfo;
+	}
+
+	virtual FLandscapeHeightmapInfo Validate(const TCHAR* HeightmapFilename) const override;
+	virtual FLandscapeHeightmapImportData Import(const TCHAR* HeightmapFilename, FLandscapeFileResolution ExpectedResolution) const override;
+	virtual void Export(const TCHAR* HeightmapFilename, TArrayView<const uint16> Data, FLandscapeFileResolution DataResolution, FVector Scale) const override;
+};
+
+//////////////////////////////////////////////////////////////////////////
+
+class FLandscapeWeightmapFileFormat_Png : public ILandscapeWeightmapFileFormat
+{
+private:
+	FLandscapeFileTypeInfo FileTypeInfo;
+
+public:
+	FLandscapeWeightmapFileFormat_Png();
+
+	virtual const FLandscapeFileTypeInfo& GetInfo() const override
+	{
+		return FileTypeInfo;
+	}
+
+	virtual FLandscapeWeightmapInfo Validate(const TCHAR* WeightmapFilename, FName LayerName) const override;
+	virtual FLandscapeWeightmapImportData Import(const TCHAR* WeightmapFilename, FName LayerName, FLandscapeFileResolution ExpectedResolution) const override;
+	virtual void Export(const TCHAR* WeightmapFilename, FName LayerName, TArrayView<const uint8> Data, FLandscapeFileResolution DataResolution) const override;
+};
diff --git a/Engine/Source/Editor/LandscapeEditor/Private/LandscapeFileFormatRaw.cpp b/Engine/Source/Editor/LandscapeEditor/Private/LandscapeFileFormatRaw.cpp
new file mode 100644
--- /dev/null
+++ b/Engine/Source/Editor/LandscapeEditor/Private/LandscapeFileFormatRaw.cpp
@@ -0,0 +1,173 @@
+// Copyright 1998-2016 Epic Games, Inc. All Rights Reserved.
+
+#include "LandscapeEditorPrivatePCH.h"
+#include "LandscapeFileFormatRaw.h"
+
+#define LOCTEXT_NAMESPACE "LandscapeEditor.NewLandscape"
+
+TArray<FLandscapeFileResolution> CalculatePossibleRawResolutions(int64 FileSize)
+{
+	TArray<FLandscapeFileResolution> PossibleResolutions;
+
+	// Find all possible heightmap sizes, between 8 and 8192 width/height
+	const int32 MinWidth = FMath::Max(8, (int32)FMath::DivideAndRoundUp(FileSize, (int64)8192));
+	const int32 MaxWidth = FMath::TruncToInt(FMath::Sqrt(FileSize));
+	for (int32 Width = MinWidth; Width <= MaxWidth; Width++)
+	{
+		if (FileSize % Width == 0)
+		{
+			FLandscapeFileResolution ImportResolution;
+			ImportResolution.Width = Width;
+			ImportResolution.Height = FileSize / Width;
+			PossibleResolutions.Add(ImportResolution);
+		}
+	}
+
+	for (int32 i = PossibleResolutions.Num() - 1; i >= 0; --i)
+	{
+		FLandscapeFileResolution ImportResolution = PossibleResolutions[i];
+		if (ImportResolution.Width != ImportResolution.Height)
+		{
+			Swap(ImportResolution.Width, ImportResolution.Height);
+			PossibleResolutions.Add(ImportResolution);
+		}
+	}
+
+	return PossibleResolutions;
+}
+
+FLandscapeHeightmapFileFormat_Raw::FLandscapeHeightmapFileFormat_Raw()
+{
+	FileTypeInfo.Description = LOCTEXT("FileFormatRaw_HeightmapDesc", "Heightmap .r16/.raw files");
+	FileTypeInfo.Extensions.Add(".r16");
+	FileTypeInfo.Extensions.Add(".raw");
+	FileTypeInfo.bSupportsExport = true;
+}
+
+FLandscapeHeightmapInfo FLandscapeHeightmapFileFormat_Raw::Validate(const TCHAR* HeightmapFilename) const
+{
+	FLandscapeHeightmapInfo Result;
+
+	int64 ImportFileSize = IFileManager::Get().FileSize(HeightmapFilename);
+
+	if (ImportFileSize < 0)
+	{
+		Result.ResultCode = ELandscapeImportResult::Error;
+		Result.ErrorMessage = LOCTEXT("Import_HeightmapFileReadError", "Error reading heightmap file");
+	}
+	else if (ImportFileSize == 0 || ImportFileSize % 2 != 0)
+	{
+		Result.ResultCode = ELandscapeImportResult::Error;
+		Result.ErrorMessage = LOCTEXT("Import_HeightmapFileInvalidSize", "The heightmap file has an invalid size (possibly not 16-bit?)");
+	}
+	else
+	{
+		Result.PossibleResolutions = CalculatePossibleRawResolutions(ImportFileSize / 2);
+
+		if (Result.PossibleResolutions.Num() == 0)
+		{
+			Result.ResultCode = ELandscapeImportResult::Error;
+			Result.ErrorMessage = LOCTEXT("Import_HeightmapFileInvalidSize", "The heightmap file has an invalid size (possibly not 16-bit?)");
+		}
+	}
+
+	return Result;
+}
+
+FLandscapeHeightmapImportData FLandscapeHeightmapFileFormat_Raw::Import(const TCHAR* HeightmapFilename, FLandscapeFileResolution ExpectedResolution) const
+{
+	FLandscapeHeightmapImportData Result;
+
+	TArray<uint8> TempData;
+	if (!FFileHelper::LoadFileToArray(TempData, HeightmapFilename, FILEREAD_Silent))
+	{
+		Result.ResultCode = ELandscapeImportResult::Error;
+		Result.ErrorMessage = LOCTEXT("Import_HeightmapFileReadError", "Error reading heightmap file");
+	}
+	else if (TempData.Num() != (ExpectedResolution.Width * ExpectedResolution.Height * 2))
+	{
+		Result.ResultCode = ELandscapeImportResult::Error;
+		Result.ErrorMessage = LOCTEXT("Import_HeightmapResolutionMismatch", "The heightmap file's resolution does not match the requested resolution");
+	}
+	else
+	{
+		Result.Data.Empty(ExpectedResolution.Width * ExpectedResolution.Height);
+		Result.Data.AddUninitialized(ExpectedResolution.Width * ExpectedResolution.Height);
+		FMemory::Memcpy(Result.Data.GetData(), TempData.GetData(), ExpectedResolution.Width * ExpectedResolution.Height * 2);
+	}
+
+	return Result;
+}
+
+void FLandscapeHeightmapFileFormat_Raw::Export(const TCHAR* HeightmapFilename, TArrayView<const uint16> Data, FLandscapeFileResolution DataResolution, FVector Scale) const
+{
+	TArray<uint8> TempData;
+	TempData.Empty(DataResolution.Width * DataResolution.Height * 2);
+	TempData.AddUninitialized(DataResolution.Width * DataResolution.Height * 2);
+	FMemory::Memcpy(TempData.GetData(), Data.GetData(), DataResolution.Width * DataResolution.Height * 2);
+
+	FFileHelper::SaveArrayToFile(TempData, HeightmapFilename);
+}
+
+//////////////////////////////////////////////////////////////////////////
+
+FLandscapeWeightmapFileFormat_Raw::FLandscapeWeightmapFileFormat_Raw()
+{
+	FileTypeInfo.Description = LOCTEXT("FileFormatRaw_WeightmapDesc", "Layer .r8/.raw files");
+	FileTypeInfo.Extensions.Add(".r8");
+	FileTypeInfo.Extensions.Add(".raw");
+	FileTypeInfo.bSupportsExport = true;
+}
+
+FLandscapeWeightmapInfo FLandscapeWeightmapFileFormat_Raw::Validate(const TCHAR* WeightmapFilename, FName LayerName) const
+{
+	FLandscapeWeightmapInfo Result;
+
+	int64 ImportFileSize = IFileManager::Get().FileSize(WeightmapFilename);
+
+	if (ImportFileSize < 0)
+	{
+		Result.ResultCode = ELandscapeImportResult::Error;
+		Result.ErrorMessage = LOCTEXT("Import_LayerFileReadError", "Error reading layer file");
+	}
+	else
+	{
+		Result.PossibleResolutions = CalculatePossibleRawResolutions(ImportFileSize / 2);
+
+		if (Result.PossibleResolutions.Num() == 0)
+		{
+			Result.ResultCode = ELandscapeImportResult::Error;
+			Result.ErrorMessage = LOCTEXT("Import_WeightmapFileInvalidSize", "The layer file has an invalid size");
+		}
+	}
+
+	return Result;
+}
+
+FLandscapeWeightmapImportData FLandscapeWeightmapFileFormat_Raw::Import(const TCHAR* WeightmapFilename, FName LayerName, FLandscapeFileResolution ExpectedResolution) const
+{
+	FLandscapeWeightmapImportData Result;
+
+	TArray<uint8> TempData;
+	if (!FFileHelper::LoadFileToArray(TempData, WeightmapFilename, FILEREAD_Silent))
+	{
+		Result.ResultCode = ELandscapeImportResult::Error;
+		Result.ErrorMessage = LOCTEXT("Import_LayerFileReadError", "Error reading layer file");
+	}
+	else if (TempData.Num() != (ExpectedResolution.Width * ExpectedResolution.Height))
+	{
+		Result.ResultCode = ELandscapeImportResult::Error;
+		Result.ErrorMessage = LOCTEXT("Import_LayerResolutionMismatch", "The layer file's resolution does not match the requested resolution");
+	}
+	else
+	{
+		Result.Data = MoveTemp(TempData);
+	}
+
+	return Result;
+}
+
+void FLandscapeWeightmapFileFormat_Raw::Export(const TCHAR* WeightmapFilename, FName LayerName, TArrayView<const uint8> Data, FLandscapeFileResolution DataResolution) const
+{
+	FFileHelper::SaveArrayToFile(Data, WeightmapFilename);
+}
diff --git a/Engine/Source/Editor/LandscapeEditor/Private/LandscapeFileFormatRaw.h b/Engine/Source/Editor/LandscapeEditor/Private/LandscapeFileFormatRaw.h
new file mode 100644
--- /dev/null
+++ b/Engine/Source/Editor/LandscapeEditor/Private/LandscapeFileFormatRaw.h
@@ -0,0 +1,44 @@
+// Copyright 1998-2016 Epic Games, Inc. All Rights Reserved.
+
+#pragma once
+
+#include "LandscapeFileFormatInterface.h"
+
+// Implement .raw file format
+class FLandscapeHeightmapFileFormat_Raw : public ILandscapeHeightmapFileFormat
+{
+private:
+	FLandscapeFileTypeInfo FileTypeInfo;
+
+public:
+	FLandscapeHeightmapFileFormat_Raw();
+
+	virtual const FLandscapeFileTypeInfo& GetInfo() const override
+	{
+		return FileTypeInfo;
+	}
+
+	virtual FLandscapeHeightmapInfo Validate(const TCHAR* HeightmapFilename) const override;
+	virtual FLandscapeHeightmapImportData Import(const TCHAR* HeightmapFilename, FLandscapeFileResolution ExpectedResolution) const override;
+	virtual void Export(const TCHAR* HeightmapFilename, TArrayView<const uint16> Data, FLandscapeFileResolution DataResolution, FVector Scale) const override;
+};
+
+//////////////////////////////////////////////////////////////////////////
+
+class FLandscapeWeightmapFileFormat_Raw : public ILandscapeWeightmapFileFormat
+{
+private:
+	FLandscapeFileTypeInfo FileTypeInfo;
+
+public:
+	FLandscapeWeightmapFileFormat_Raw();
+
+	virtual const FLandscapeFileTypeInfo& GetInfo() const override
+	{
+		return FileTypeInfo;
+	}
+
+	virtual FLandscapeWeightmapInfo Validate(const TCHAR* WeightmapFilename, FName LayerName) const override;
+	virtual FLandscapeWeightmapImportData Import(const TCHAR* WeightmapFilename, FName LayerName, FLandscapeFileResolution ExpectedResolution) const override;
+	virtual void Export(const TCHAR* WeightmapFilename, FName LayerName, TArrayView<const uint8> Data, FLandscapeFileResolution DataResolution) const override;
+};
diff --git a/Engine/Source/Editor/LandscapeEditor/Public/LandscapeEditorModule.h b/Engine/Source/Editor/LandscapeEditor/Public/LandscapeEditorModule.h
--- a/Engine/Source/Editor/LandscapeEditor/Public/LandscapeEditorModule.h
+++ b/Engine/Source/Editor/LandscapeEditor/Public/LandscapeEditorModule.h
@@ -2,6 +2,7 @@
 
 #pragma once
 
+#include "LandscapeFileFormatInterface.h"
 
 /**
  * LandscapeEditor module interface
@@ -9,5 +10,19 @@
 class ILandscapeEditorModule : public IModuleInterface
 {
 public:
-};
+	// Register / unregister a landscape file format plugin
+	virtual void RegisterHeightmapFileFormat(TSharedRef<ILandscapeHeightmapFileFormat> FileFormat) = 0;
+	virtual void RegisterWeightmapFileFormat(TSharedRef<ILandscapeWeightmapFileFormat> FileFormat) = 0;
+	virtual void UnregisterHeightmapFileFormat(TSharedRef<ILandscapeHeightmapFileFormat> FileFormat) = 0;
+	virtual void UnregisterWeightmapFileFormat(TSharedRef<ILandscapeWeightmapFileFormat> FileFormat) = 0;
+
+	// Gets the type string used by the import/export file dialog
+	virtual const TCHAR* GetHeightmapImportDialogTypeString() const = 0;
+	virtual const TCHAR* GetWeightmapImportDialogTypeString() const = 0;
+	virtual const TCHAR* GetHeightmapExportDialogTypeString() const = 0;
+	virtual const TCHAR* GetWeightmapExportDialogTypeString() const = 0;
 
+	// Gets the heightmap/weightmap format associated with a given extension (null if no plugin is registered for this extension)
+	virtual const ILandscapeHeightmapFileFormat* GetHeightmapFormatByExtension(const TCHAR* Extension) = 0;
+	virtual const ILandscapeWeightmapFileFormat* GetWeightmapFormatByExtension(const TCHAR* Extension) = 0;
+};
diff --git a/Engine/Source/Editor/LandscapeEditor/Public/LandscapeEditorObject.h b/Engine/Source/Editor/LandscapeEditor/Public/LandscapeEditorObject.h
--- a/Engine/Source/Editor/LandscapeEditor/Public/LandscapeEditorObject.h
+++ b/Engine/Source/Editor/LandscapeEditor/Public/LandscapeEditorObject.h
@@ -3,6 +3,7 @@
 #pragma once
 
 #include "LandscapeProxy.h" // for ELandscapeLayerPaintingRestriction
+#include "LandscapeFileFormatInterface.h"
 #include "LandscapeEditorObject.generated.h"
 
 UENUM()
@@ -190,12 +191,16 @@ struct FLandscapeImportLayer : public FLandscapeImportLayerInfo
 	ULandscapeMaterialInstanceConstant* ThumbnailMIC;
 
 	UPROPERTY(Category="Import", VisibleAnywhere)
-	ELandscapeImportLayerError ImportError;
+	ELandscapeImportResult ImportResult;
+
+	UPROPERTY(Category="Import", VisibleAnywhere)
+	FText ErrorMessage;
 
 	FLandscapeImportLayer()
 		: FLandscapeImportLayerInfo()
 		, ThumbnailMIC(nullptr)
-		, ImportError(ELandscapeImportLayerError::None)
+		, ImportResult(ELandscapeImportResult::Success)
+		, ErrorMessage(FText())
 	{
 	}
 };
@@ -464,15 +469,18 @@ class ULandscapeEditorObject : public UObject
 	FVector NewLandscape_Scale;
 
 	UPROPERTY(Category="New Landscape", VisibleAnywhere, NonTransactional, meta=(ShowForTools="NewLandscape"))
-	TEnumAsByte<ELandscapeImportHeightmapError::Type> ImportLandscape_HeightmapError;
+	ELandscapeImportResult ImportLandscape_HeightmapImportResult;
+
+	UPROPERTY(Category="New Landscape", VisibleAnywhere, NonTransactional, meta=(ShowForTools="NewLandscape"))
+	FText ImportLandscape_HeightmapErrorMessage;
 
 	// Specify a height map file in 16-bit RAW or PNG format
 	UPROPERTY(Category="New Landscape", EditAnywhere, NonTransactional, meta=(DisplayName="Heightmap File", ShowForTools="NewLandscape"))
 	FString ImportLandscape_HeightmapFilename;
 	UPROPERTY(NonTransactional)
-	int32 ImportLandscape_Width;
+	uint32 ImportLandscape_Width;
 	UPROPERTY(NonTransactional)
-	int32 ImportLandscape_Height;
+	uint32 ImportLandscape_Height;
 	UPROPERTY(NonTransactional)
 	TArray<uint16> ImportLandscape_Data;
 
diff --git a/Engine/Source/Editor/LandscapeEditor/Public/LandscapeFileFormatInterface.h b/Engine/Source/Editor/LandscapeEditor/Public/LandscapeFileFormatInterface.h
new file mode 100644
--- /dev/null
+++ b/Engine/Source/Editor/LandscapeEditor/Public/LandscapeFileFormatInterface.h
@@ -0,0 +1,170 @@
+// Copyright 1998-2016 Epic Games, Inc. All Rights Reserved.
+
+#pragma once
+
+#include "ArrayView.h"
+
+struct FLandscapeFileTypeInfo
+{
+	// Description of file type for the file selector
+	FText Description;
+
+	// Extensions for this type, with leading dot, e.g. ".png"
+	TArray<FString, TInlineAllocator<2>> Extensions;
+
+	// Whether this file type supports exporting from the editor back to file
+	// (All file types must support *importing*, but exporting is optional)
+	bool bSupportsExport = false;
+};
+
+UENUM()
+enum class ELandscapeImportResult : uint8
+{
+	Success = 0,
+	Warning,
+	Error,
+};
+
+struct FLandscapeFileResolution
+{
+	uint32 Width;
+	uint32 Height;
+};
+
+FORCEINLINE bool operator==(const FLandscapeFileResolution& Lhs, const FLandscapeFileResolution& Rhs)
+{
+	return Lhs.Width == Rhs.Width && Lhs.Height == Rhs.Height;
+}
+
+FORCEINLINE bool operator!=(const FLandscapeFileResolution& Lhs, const FLandscapeFileResolution& Rhs)
+{
+	return !(Lhs == Rhs);
+}
+
+struct FLandscapeHeightmapInfo
+{
+	// Whether the heightmap is usable or has errors/warnings
+	ELandscapeImportResult ResultCode = ELandscapeImportResult::Success;
+
+	// Message to show as the warning/error result
+	FText ErrorMessage;
+
+	// Normally contains a single resolution, but .raw is awful
+	TArray<FLandscapeFileResolution, TInlineAllocator<1>> PossibleResolutions;
+
+	// The inherent scale of the data format, if it has one, in centimeters
+	// The default for data with no inherent scale is 100,100,0.78125 (100.0/128, shown as 100 in the editor UI)
+	TOptional<FVector> DataScale;
+};
+
+struct FLandscapeWeightmapInfo
+{
+	// Whether the weightmap is usable or has errors/warnings
+	ELandscapeImportResult ResultCode = ELandscapeImportResult::Success;
+
+	// Message to show as the warning/error result
+	FText ErrorMessage;
+
+	// Normally contains a single resolution, but .raw is awful
+	TArray<FLandscapeFileResolution, TInlineAllocator<1>> PossibleResolutions;
+};
+
+struct FLandscapeHeightmapImportData
+{
+	// Whether the heightmap is usable or has errors/warnings
+	ELandscapeImportResult ResultCode = ELandscapeImportResult::Success;
+
+	// Message to show as the warning/error result
+	FText ErrorMessage;
+
+	// The height data!
+	// A value of 32768 is the 0 level (e.g. sea level), lower values are below and higher values are above
+	TArray<uint16> Data;
+};
+
+struct FLandscapeWeightmapImportData
+{
+	// Whether the weightmap is usable or has errors/warnings
+	ELandscapeImportResult ResultCode = ELandscapeImportResult::Success;
+
+	// Message to show as the warning/error result
+	FText ErrorMessage;
+
+	// The weight data!
+	// 255 is fully painted and 0 is unpainted
+	TArray<uint8> Data;
+};
+
+// Interface
+class ILandscapeHeightmapFileFormat
+{
+public:
+	/** Gets info about this format
+	 * @return information about the file types supported by this file format plugin
+	 */
+	virtual const FLandscapeFileTypeInfo& GetInfo() const = 0;
+
+	/** Validate a file for Import
+	 * Gives the file format the opportunity to reject a file or return warnings
+	 * as well as return information about the file for the import UI (e.g. resolution and scale)
+	 * @param HeightmapFilename path to the file to validate for import
+	 * @return information about the file and (optional) error message
+	 */
+	virtual FLandscapeHeightmapInfo Validate(const TCHAR* HeightmapFilename) const = 0;
+
+	/** Import a file
+	 * @param HeightmapFilename path to the file to import
+	 * @param ExpectedResolution resolution selected in the import UI (mostly for the benefit of .raw)
+	 * @return imported data and (optional) error message
+	 */
+	virtual FLandscapeHeightmapImportData Import(const TCHAR* HeightmapFilename, FLandscapeFileResolution ExpectedResolution) const = 0;
+
+	/** Export a file (if supported)
+	 * @param HeightmapFilename path to the file to export to
+	 * @param Data raw data to export
+	 * @param DataResolution resolution of Data
+	 * @param Scale scale of the landscape data, in centimeters
+	 */
+	virtual void Export(const TCHAR* HeightmapFilename, TArrayView<const uint16> Data, FLandscapeFileResolution DataResolution, FVector Scale) const
+	{
+		checkf(0, TEXT("File type hasn't implemented support for heightmap export - %s"), FPaths::GetExtension(HeightmapFilename, true));
+	}
+};
+
+class ILandscapeWeightmapFileFormat
+{
+public:
+	/** Gets info about this format
+	 * @return information about the file types supported by this file format plugin
+	 */
+	virtual const FLandscapeFileTypeInfo& GetInfo() const = 0;
+
+	/** Validate a file for Import
+	 * Gives the file format the opportunity to reject a file or return warnings
+	 * as well as return information about the file for the import UI (e.g. resolution and scale)
+	 * @param HeightmapFilename path to the file to validate for import
+	 * @param LayerName name of layer that is being imported
+	 * @return information about the file and (optional) error message
+	 */
+	virtual FLandscapeWeightmapInfo Validate(const TCHAR* WeightmapFilename, FName LayerName) const = 0;
+
+	/** Import a file
+	 * @param HeightmapFilename path to the file to import
+	 * @param LayerName name of layer that is being imported
+	 * @param ExpectedResolution resolution selected in the import UI (mostly for the benefit of .raw)
+	 * @return imported data and (optional) error message
+	 */
+	virtual FLandscapeWeightmapImportData Import(const TCHAR* WeightmapFilename, FName LayerName, FLandscapeFileResolution ExpectedResolution) const = 0;
+
+	/** Export a file (if supported)
+	 * @param HeightmapFilename path to the file to export to
+	 * @param LayerName name of layer that is being exported
+	 * @param Data raw data to export
+	 * @param DataResolution resolution of Data
+	 * @param Scale scale of the landscape data, in centimeters
+	 */
+	virtual void Export(const TCHAR* WeightmapFilename, FName LayerName, TArrayView<const uint8> Data, FLandscapeFileResolution DataResolution) const
+	{
+		checkf(0, TEXT("File type hasn't implemented support for weightmap export - %s"), FPaths::GetExtension(WeightmapFilename, true));
+	}
+};
diff --git a/Engine/Source/Editor/PropertyEditor/Private/PropertyEditorHelpers.cpp b/Engine/Source/Editor/PropertyEditor/Private/PropertyEditorHelpers.cpp
--- a/Engine/Source/Editor/PropertyEditor/Private/PropertyEditorHelpers.cpp
+++ b/Engine/Source/Editor/PropertyEditor/Private/PropertyEditorHelpers.cpp
@@ -470,6 +470,10 @@ namespace PropertyEditorHelpers
 		{
 			PropertyHandle = MakeShareable( new FPropertyHandleString( PropertyNode, NotifyHook, PropertyUtilities ) );
 		}
+		else if (FPropertyHandleText::Supports(PropertyNode))
+		{
+			PropertyHandle = MakeShareable(new FPropertyHandleText(PropertyNode, NotifyHook, PropertyUtilities));
+		}
 		else if( FPropertyHandleVector::Supports( PropertyNode ) )
 		{
 			PropertyHandle = MakeShareable( new FPropertyHandleVector( PropertyNode, NotifyHook, PropertyUtilities ) );
diff --git a/Engine/Source/Editor/PropertyEditor/Private/PropertyHandleImpl.cpp b/Engine/Source/Editor/PropertyEditor/Private/PropertyHandleImpl.cpp
--- a/Engine/Source/Editor/PropertyEditor/Private/PropertyHandleImpl.cpp
+++ b/Engine/Source/Editor/PropertyEditor/Private/PropertyHandleImpl.cpp
@@ -1496,6 +1496,7 @@ IMPLEMENT_PROPERTY_ACCESSOR( uint32 )
 IMPLEMENT_PROPERTY_ACCESSOR( uint64 )
 IMPLEMENT_PROPERTY_ACCESSOR( float )
 IMPLEMENT_PROPERTY_ACCESSOR( FString )
+IMPLEMENT_PROPERTY_ACCESSOR( FText )
 IMPLEMENT_PROPERTY_ACCESSOR( FName )
 IMPLEMENT_PROPERTY_ACCESSOR( FVector )
 IMPLEMENT_PROPERTY_ACCESSOR( FVector2D )
@@ -2236,6 +2237,7 @@ IMPLEMENT_PROPERTY_VALUE( FPropertyHandleByte )
 IMPLEMENT_PROPERTY_VALUE( FPropertyHandleString )
 IMPLEMENT_PROPERTY_VALUE( FPropertyHandleObject )
 IMPLEMENT_PROPERTY_VALUE( FPropertyHandleArray )
+IMPLEMENT_PROPERTY_VALUE( FPropertyHandleText )
 
 // int32 
 bool FPropertyHandleInt::Supports( TSharedRef<FPropertyNode> PropertyNode )
@@ -3140,3 +3142,34 @@ bool FPropertyHandleArray::IsEditable() const
 	// Property is editable if its a non-const dynamic array
 	return Implementation->HasValidPropertyNode() && !Implementation->IsEditConst() && Implementation->IsPropertyTypeOf(UArrayProperty::StaticClass());
 }
+
+// Localized Text
+bool FPropertyHandleText::Supports(TSharedRef<FPropertyNode> PropertyNode)
+{
+	UProperty* Property = PropertyNode->GetProperty();
+
+	if ( Property == nullptr )
+	{
+		return false;
+	}
+
+	// Supported if the property is a text property only
+	return Property->IsA(UTextProperty::StaticClass());
+}
+
+FPropertyAccess::Result FPropertyHandleText::GetValue(FText& OutValue) const
+{
+	return Implementation->GetValueAsText(OutValue);
+}
+
+FPropertyAccess::Result FPropertyHandleText::SetValue(const FText& NewValue, EPropertyValueSetFlags::Type Flags)
+{
+	FString StringValue;
+	FTextStringHelper::WriteToString(StringValue, NewValue);
+	return Implementation->ImportText(StringValue, Flags);
+}
+
+FPropertyAccess::Result FPropertyHandleText::SetValue(const FString& NewValue, EPropertyValueSetFlags::Type Flags)
+{
+	return SetValue(FText::FromString(NewValue), Flags);
+}
diff --git a/Engine/Source/Editor/PropertyEditor/Private/PropertyHandleImpl.h b/Engine/Source/Editor/PropertyEditor/Private/PropertyHandleImpl.h
--- a/Engine/Source/Editor/PropertyEditor/Private/PropertyHandleImpl.h
+++ b/Engine/Source/Editor/PropertyEditor/Private/PropertyHandleImpl.h
@@ -372,6 +372,7 @@ public:
 	DECLARE_PROPERTY_ACCESSOR( uint64 )
 	DECLARE_PROPERTY_ACCESSOR( float )
 	DECLARE_PROPERTY_ACCESSOR( FString )
+	DECLARE_PROPERTY_ACCESSOR( FText )
 	DECLARE_PROPERTY_ACCESSOR( FName )
 	DECLARE_PROPERTY_ACCESSOR( FVector )
 	DECLARE_PROPERTY_ACCESSOR( FVector2D )
@@ -593,3 +594,14 @@ private:
 	 */
 	bool IsEditable() const;
 };
+
+class FPropertyHandleText : public FPropertyHandleBase
+{
+public:
+	FPropertyHandleText(TSharedRef<FPropertyNode> PropertyNode, FNotifyHook* NotifyHook,TSharedPtr<IPropertyUtilities> PropertyUtilities);
+	static bool Supports(TSharedRef<FPropertyNode> PropertyNode);
+	virtual FPropertyAccess::Result GetValue(FText& OutValue) const override;
+	virtual FPropertyAccess::Result SetValue(const FText& InValue, EPropertyValueSetFlags::Type Flags = EPropertyValueSetFlags::DefaultFlags) override;
+
+	virtual FPropertyAccess::Result SetValue(const FString& InValue, EPropertyValueSetFlags::Type Flags = EPropertyValueSetFlags::DefaultFlags) override;
+};
diff --git a/Engine/Source/Editor/PropertyEditor/Public/PropertyHandle.h b/Engine/Source/Editor/PropertyEditor/Public/PropertyHandle.h
--- a/Engine/Source/Editor/PropertyEditor/Public/PropertyHandle.h
+++ b/Engine/Source/Editor/PropertyEditor/Public/PropertyHandle.h
@@ -195,6 +195,7 @@ public:
 	virtual FPropertyAccess::Result GetValue( uint32& OutValue ) const = 0;
 	virtual FPropertyAccess::Result GetValue( uint64& OutValue ) const = 0;
 	virtual FPropertyAccess::Result GetValue( FString& OutValue ) const = 0;
+	virtual FPropertyAccess::Result GetValue( FText& OutValue ) const = 0;
 	virtual FPropertyAccess::Result GetValue( FName& OutValue ) const = 0;
 	virtual FPropertyAccess::Result GetValue( FVector& OutValue ) const = 0;
 	virtual FPropertyAccess::Result GetValue( FVector2D& OutValue ) const = 0;
@@ -223,6 +224,7 @@ public:
 	virtual FPropertyAccess::Result SetValue( const uint32& InValue,  EPropertyValueSetFlags::Type Flags = EPropertyValueSetFlags::DefaultFlags ) = 0;
 	virtual FPropertyAccess::Result SetValue( const uint64& InValue,  EPropertyValueSetFlags::Type Flags = EPropertyValueSetFlags::DefaultFlags ) = 0;
 	virtual FPropertyAccess::Result SetValue( const FString& InValue,  EPropertyValueSetFlags::Type Flags = EPropertyValueSetFlags::DefaultFlags ) = 0;
+	virtual FPropertyAccess::Result SetValue( const FText& InValue,  EPropertyValueSetFlags::Type Flags = EPropertyValueSetFlags::DefaultFlags ) = 0;
 	virtual FPropertyAccess::Result SetValue( const FName& InValue, EPropertyValueSetFlags::Type Flags = EPropertyValueSetFlags::DefaultFlags ) = 0;
 	virtual FPropertyAccess::Result SetValue( const FVector& InValue,  EPropertyValueSetFlags::Type Flags = EPropertyValueSetFlags::DefaultFlags ) = 0;
 	virtual FPropertyAccess::Result SetValue( const FVector2D& InValue,  EPropertyValueSetFlags::Type Flags = EPropertyValueSetFlags::DefaultFlags ) = 0;
diff --git a/Engine/Source/Editor/SequenceRecorder/Private/SequenceRecorderModule.cpp b/Engine/Source/Editor/SequenceRecorder/Private/SequenceRecorderModule.cpp
--- a/Engine/Source/Editor/SequenceRecorder/Private/SequenceRecorderModule.cpp
+++ b/Engine/Source/Editor/SequenceRecorder/Private/SequenceRecorderModule.cpp
@@ -481,7 +481,7 @@ class FSequenceRecorderModule : public ISequenceRecorder, private FSelfRegisteri
 		return CurrentSequence.IsValid() ? CurrentSequence->GetMovieScene()->GetPlaybackRange().Size<float>() : 0.0f;
 	}
 
-	virtual bool StartRecording(TFixedSizeArrayView<AActor*> ActorsToRecord, const FOnRecordingStarted& OnRecordingStarted, const FOnRecordingFinished& OnRecordingFinished, const FString& PathToRecordTo, const FString& SequenceName) override
+	virtual bool StartRecording(TArrayView<AActor* const> ActorsToRecord, const FOnRecordingStarted& OnRecordingStarted, const FOnRecordingFinished& OnRecordingFinished, const FString& PathToRecordTo, const FString& SequenceName) override
 	{
 		if(ActorsToRecord.Num() != 0)
 		{
diff --git a/Engine/Source/Editor/SequenceRecorder/Public/ISequenceRecorder.h b/Engine/Source/Editor/SequenceRecorder/Public/ISequenceRecorder.h
--- a/Engine/Source/Editor/SequenceRecorder/Public/ISequenceRecorder.h
+++ b/Engine/Source/Editor/SequenceRecorder/Public/ISequenceRecorder.h
@@ -3,7 +3,7 @@
 #pragma once
 
 #include "ModuleInterface.h"
-#include "FixedSizeArrayView.h"
+#include "Containers/ArrayView.h"
 
 DECLARE_DELEGATE_OneParam(FOnRecordingStarted, class UMovieSceneSequence* /*Sequence*/);
 
@@ -38,7 +38,7 @@ public:
 	 * @param	SequenceName		Optional name of a sequence to record to. If none is specified we use the defaults in the settings.
 	 * @return true if recording was successfully started
 	*/
-	virtual bool StartRecording(TFixedSizeArrayView<AActor*> ActorsToRecord, const FOnRecordingStarted& OnRecordingStarted, const FOnRecordingFinished& OnRecordingFinished, const FString& PathToRecordTo = FString(), const FString& SequenceName = FString()) = 0;
+	virtual bool StartRecording(TArrayView<AActor* const> ActorsToRecord, const FOnRecordingStarted& OnRecordingStarted, const FOnRecordingFinished& OnRecordingFinished, const FString& PathToRecordTo = FString(), const FString& SequenceName = FString()) = 0;
 
 	/**
 	 * Start a recording, possibly with some delay (specified by the sequence recording settings).
@@ -51,7 +51,7 @@ public:
 	*/
 	bool StartRecording(AActor* ActorToRecord, const FOnRecordingStarted& OnRecordingStarted, const FOnRecordingFinished& OnRecordingFinished, const FString& PathToRecordTo = FString(), const FString& SequenceName = FString())
 	{
-		return StartRecording(TFixedSizeArrayView<AActor*>(&ActorToRecord, 1), OnRecordingStarted, OnRecordingFinished, PathToRecordTo, SequenceName);
+		return StartRecording(MakeArrayView(&ActorToRecord, 1), OnRecordingStarted, OnRecordingFinished, PathToRecordTo, SequenceName);
 	}
 
 	/**
diff --git a/Engine/Source/Editor/UnrealEd/Classes/Editor/UnrealEdEngine.h b/Engine/Source/Editor/UnrealEd/Classes/Editor/UnrealEdEngine.h
--- a/Engine/Source/Editor/UnrealEd/Classes/Editor/UnrealEdEngine.h
+++ b/Engine/Source/Editor/UnrealEd/Classes/Editor/UnrealEdEngine.h
@@ -760,8 +760,6 @@ public:
 	bool HandleDumpSelectionCommand( const TCHAR* Str, FOutputDevice& Ar );
 	bool HandleBuildLightingCommand( const TCHAR* Str, FOutputDevice& Ar, UWorld* InWorld );
 	bool HandleBuildPathsCommand( const TCHAR* Str, FOutputDevice& Ar, UWorld* InWorld );
-	bool HandleUpdateLandscapeEditorDataCommand( const TCHAR* Str, FOutputDevice& Ar, UWorld* InWorld );
-	bool HandleUpdateLandscapeMICCommand( const TCHAR* Str, FOutputDevice& Ar, UWorld* InWorld );
 	bool HandleRecreateLandscapeCollisionCommand(const TCHAR* Str, FOutputDevice& Ar, UWorld* InWorld);
 	bool HandleRemoveLandscapeXYOffsetsCommand(const TCHAR* Str, FOutputDevice& Ar, UWorld* InWorld);
 	bool HandleConvertMatineesCommand( const TCHAR* Str, FOutputDevice& Ar, UWorld* InWorld );
diff --git a/Engine/Source/Editor/UnrealEd/Private/EditorHook.cpp b/Engine/Source/Editor/UnrealEd/Private/EditorHook.cpp
--- a/Engine/Source/Editor/UnrealEd/Private/EditorHook.cpp
+++ b/Engine/Source/Editor/UnrealEd/Private/EditorHook.cpp
@@ -25,16 +25,16 @@ UEngine* Engine;
 	Editor hook exec.
 -----------------------------------------------------------------------------*/
 
-void UUnrealEdEngine::NotifyPreChange( UProperty* PropertyAboutToChange )
+void UUnrealEdEngine::NotifyPreChange(UProperty* PropertyAboutToChange)
 {
 }
-void UUnrealEdEngine::NotifyPostChange( const FPropertyChangedEvent& PropertyChangedEvent, UProperty* PropertyThatChanged )
+void UUnrealEdEngine::NotifyPostChange(const FPropertyChangedEvent& PropertyChangedEvent, UProperty* PropertyThatChanged)
 {
 	// Notify all active modes of actor property changes.
 	TArray<FEdMode*> ActiveModes;
-	GLevelEditorModeTools().GetActiveModes( ActiveModes );
+	GLevelEditorModeTools().GetActiveModes(ActiveModes);
 
-	for( int32 ModeIndex = 0; ModeIndex < ActiveModes.Num(); ++ModeIndex )
+	for (int32 ModeIndex = 0; ModeIndex < ActiveModes.Num(); ++ModeIndex)
 	{
 		ActiveModes[ModeIndex]->ActorPropChangeNotify();
 	}
@@ -44,40 +44,15 @@ void UUnrealEdEngine::UpdateFloatingPropertyWindows(bool bForceRefresh)
 {
 	TArray<UObject*> SelectedObjects;
 
-	bool bProcessed = false;
-	if (GetSelectedActorCount() == 1)
+	// Assemble a set of valid selected actors.
+	for (FSelectionIterator It = GetSelectedActorIterator(); It; ++It)
 	{
-		FSelectionIterator It( GetSelectedActorIterator() );
-		if ( It )
-		{
-			ALandscapeProxy* LP = Cast<ALandscapeProxy>(*It);
-			if (LP != NULL && !LP->IsPendingKill())
-			{
-				ULandscapeInfo* Info = LP->GetLandscapeInfo(false);
-				if (Info && Info->bCurrentlyEditing)
-				{
-					for (ULandscapeComponent* LandscapeComponent : Info->GetSelectedComponents())
-					{
-						SelectedObjects.Add(LandscapeComponent);
-					}
-					bProcessed = true;
-				}
-			}
-		}
-	}
-	
-	if (!bProcessed)
-	{
-		// Assemble a set of valid selected actors.
-		for ( FSelectionIterator It( GetSelectedActorIterator() ) ; It ; ++It )
-		{
-			AActor* Actor = static_cast<AActor*>( *It );
-			checkSlow( Actor->IsA(AActor::StaticClass()) );
+		AActor* Actor = static_cast<AActor*>(*It);
+		checkSlow(Actor->IsA(AActor::StaticClass()));
 
-			if ( !Actor->IsPendingKill() )
-			{
-				SelectedObjects.Add( Actor );
-			}
+		if (!Actor->IsPendingKill())
+		{
+			SelectedObjects.Add(Actor);
 		}
 	}
 
@@ -85,7 +60,7 @@ void UUnrealEdEngine::UpdateFloatingPropertyWindows(bool bForceRefresh)
 }
 
 
-void UUnrealEdEngine::UpdateFloatingPropertyWindowsFromActorList(const TArray< UObject *>& ActorList, bool bForceRefresh)
+void UUnrealEdEngine::UpdateFloatingPropertyWindowsFromActorList(const TArray<UObject*>& ActorList, bool bForceRefresh)
 {
 	FLevelEditorModule& LevelEditor = FModuleManager::LoadModuleChecked<FLevelEditorModule>(TEXT("LevelEditor"));
 
diff --git a/Engine/Source/Editor/UnrealEd/Private/Fbx/FbxMainExport.cpp b/Engine/Source/Editor/UnrealEd/Private/Fbx/FbxMainExport.cpp
--- a/Engine/Source/Editor/UnrealEd/Private/Fbx/FbxMainExport.cpp
+++ b/Engine/Source/Editor/UnrealEd/Private/Fbx/FbxMainExport.cpp
@@ -2909,7 +2909,7 @@ void FFbxExporter::ExportInstancedMeshToFbx(const UInstancedStaticMeshComponent*
  */
 void FFbxExporter::ExportLandscapeToFbx(ALandscapeProxy* Landscape, const TCHAR* MeshName, FbxNode* FbxActor, bool bSelectedOnly)
 {
-	const ULandscapeInfo* LandscapeInfo = Landscape->GetLandscapeInfo(false);
+	const ULandscapeInfo* LandscapeInfo = Landscape->GetLandscapeInfo();
 	
 	TSet<ULandscapeComponent*> SelectedComponents;
 	if (bSelectedOnly && LandscapeInfo)
diff --git a/Engine/Source/Editor/UnrealEd/Private/Lightmass/LightmassLandscapeRender.cpp b/Engine/Source/Editor/UnrealEd/Private/Lightmass/LightmassLandscapeRender.cpp
--- a/Engine/Source/Editor/UnrealEd/Private/Lightmass/LightmassLandscapeRender.cpp
+++ b/Engine/Source/Editor/UnrealEd/Private/Lightmass/LightmassLandscapeRender.cpp
@@ -223,7 +223,7 @@ void GetLandscapeOpacityData(const FLandscapeStaticLightingMesh* LandscapeMesh,
 	const int32 LightingLOD = LandscapeComponent->GetLandscapeProxy()->StaticLightingLOD;
 	const float LightMapRatio = ::GetTerrainExpandPatchCount(LightMapRes, PatchExpandCountX, PatchExpandCountY, ComponentSizeQuads, (NumSubsections * (SubsectionSizeQuads + 1)), DesiredSize, LightingLOD);
 
-	ULandscapeInfo* const LandscapeInfo = LandscapeComponent->GetLandscapeInfo(false);
+	ULandscapeInfo* const LandscapeInfo = LandscapeComponent->GetLandscapeInfo();
 	check(LandscapeInfo);
 
 	FLandscapeEditDataInterface DataInterface = FLandscapeEditDataInterface(LandscapeInfo);
diff --git a/Engine/Source/Editor/UnrealEd/Private/UnrealEdSrv.cpp b/Engine/Source/Editor/UnrealEd/Private/UnrealEdSrv.cpp
--- a/Engine/Source/Editor/UnrealEd/Private/UnrealEdSrv.cpp
+++ b/Engine/Source/Editor/UnrealEd/Private/UnrealEdSrv.cpp
@@ -513,126 +513,6 @@ bool UUnrealEdEngine::HandleBuildPathsCommand( const TCHAR* Str, FOutputDevice&
 	return FEditorBuildUtils::EditorBuild(InWorld, FBuildOptions::BuildAIPaths);
 }
 
-bool UUnrealEdEngine::HandleUpdateLandscapeEditorDataCommand( const TCHAR* Str, FOutputDevice& Ar, UWorld* InWorld )
-{
-	const bool bShowWarnings = (FString(Str) == TEXT("-warnings"));
-	
-	if (InWorld->GetWorldSettings())
-	{
-		ULandscapeInfo::RecreateLandscapeInfo(InWorld, bShowWarnings);
-				
-		// for removing 
-		TMap<ULandscapeInfo*, ALandscapeGizmoActiveActor*> GizmoMap;
-		for (TActorIterator<ALandscapeGizmoActiveActor> It(InWorld); It; ++It)
-		{
-			ALandscapeGizmoActiveActor* Gizmo = *It;
-			if (Gizmo->TargetLandscapeInfo)
-			{
-				if (!GizmoMap.FindRef(Gizmo->TargetLandscapeInfo))
-				{
-					GizmoMap.Add(Gizmo->TargetLandscapeInfo, Gizmo);
-				}
-				else
-				{
-					Gizmo->Destroy(); 
-				}
-			}
-		}
-
-		// Fixed up for Landscape fix match case
-		auto& LandscapeInfoMap = GetLandscapeInfoMap(InWorld);
-
-		for (auto It = LandscapeInfoMap.Map.CreateIterator(); It; ++It)
-		{
-			ULandscapeInfo* LandscapeInfo = It.Value();
-			if (LandscapeInfo)
-			{
-				bool bHasPhysicalMaterial = false;
-				for (int32 i = 0; i < LandscapeInfo->Layers.Num(); ++i)
-				{
-					if (LandscapeInfo->Layers[i].LayerInfoObj && LandscapeInfo->Layers[i].LayerInfoObj->PhysMaterial)
-					{
-						bHasPhysicalMaterial = true;
-						break;
-					}
-				}
-				TSet<ALandscapeProxy*> SelectProxies;
-				for (auto LandscapeComponentIt = LandscapeInfo->XYtoComponentMap.CreateIterator(); LandscapeComponentIt; ++LandscapeComponentIt )
-				{
-					ULandscapeComponent* Comp = LandscapeComponentIt.Value();
-					if (Comp)
-					{
-						// Fix level inconsistency for landscape component and collision component
-						ULandscapeHeightfieldCollisionComponent* Collision = Comp->CollisionComponent.Get();
-						if (Collision != nullptr)
-						{
-							if (Comp->GetLandscapeProxy()->GetLevel() != Collision->GetLandscapeProxy()->GetLevel())
-							{
-								ALandscapeProxy* FromProxy = Collision->GetLandscapeProxy();
-								ALandscapeProxy* DestProxy = Comp->GetLandscapeProxy();
-								// From MoveToLevelTool
-								FromProxy->CollisionComponents.Remove(Collision);
-								Collision->UnregisterComponent();
-								Collision->DetachFromComponent(FDetachmentTransformRules::KeepWorldTransform);
-								Collision->Rename(NULL, DestProxy);
-								DestProxy->CollisionComponents.Add(Collision);
-								Collision->AttachToComponent(DestProxy->GetRootComponent(), FAttachmentTransformRules::KeepWorldTransform);
-								SelectProxies.Add(FromProxy);
-								SelectProxies.Add(DestProxy);
-							}
-
-							// Fix Dominant Layer Data
-							if (bHasPhysicalMaterial && Collision->DominantLayerData.GetBulkDataSize() == 0)
-							{
-								Comp->UpdateCollisionLayerData();
-							}
-						}
-					}
-				}
-
-				for(TSet<ALandscapeProxy*>::TIterator ProxyIt(SelectProxies);ProxyIt;++ProxyIt)
-				{
-					(*ProxyIt)->MarkPackageDirty();
-				}
-			}
-		}
-
-		// Fix proxies relative transformations to LandscapeActor
-		for (auto It = LandscapeInfoMap.Map.CreateIterator(); It; ++It)
-		{
-			ULandscapeInfo* Info = It.Value();
-			Info->FixupProxiesWeightmaps();
-			// make sure relative proxy transformations are correct	
-			Info->FixupProxiesTransform();
-		}
-	}
-	return true;
-}
-
-bool UUnrealEdEngine::HandleUpdateLandscapeMICCommand( const TCHAR* Str, FOutputDevice& Ar, UWorld* InWorld )
-{
-	UWorld* World = GetEditorWorldContext().World();
-
-	if (World && World->GetWorldSettings())
-	{
-		for (auto It = GetLandscapeInfoMap(World).Map.CreateIterator(); It; ++It)
-		{
-			ULandscapeInfo* Info = It.Value();
-			for (auto LandscapeComponentIt = Info->XYtoComponentMap.CreateIterator(); LandscapeComponentIt; ++LandscapeComponentIt )
-			{
-				ULandscapeComponent* Comp = LandscapeComponentIt.Value();
-				if( Comp )
-				{
-					Comp->UpdateMaterialInstances();
-
-					FComponentReregisterContext ReregisterContext(Comp);
-				}
-			}
-		}
-	}
-	return true;
-}
-
 bool UUnrealEdEngine::HandleRecreateLandscapeCollisionCommand(const TCHAR* Str, FOutputDevice& Ar, UWorld* InWorld)
 {
 	if (!PlayWorld && InWorld && InWorld->GetWorldSettings())
@@ -824,18 +704,6 @@ bool UUnrealEdEngine::Exec( UWorld* InWorld, const TCHAR* Stream, FOutputDevice&
 		HandleBuildPathsCommand( Str, Ar, InWorld );
 	}
 #if WITH_EDITOR
-	else if (FParse::Command(&Str, TEXT("UpdateLandscapeEditorData")))
-	{
-		// InWorld above is the PIE world if PIE is active, but this is specifically an editor command
-		UWorld* World = GetEditorWorldContext().World();
-		return HandleUpdateLandscapeEditorDataCommand(Str, Ar, World);
-	}
-	else if (FParse::Command(&Str, TEXT("UpdateLandscapeMIC")))
-	{
-		// InWorld above is the PIE world if PIE is active, but this is specifically an editor command
-		UWorld* World = GetEditorWorldContext().World();
-		return HandleUpdateLandscapeMICCommand(Str, Ar, World);
-	}
 	else if (FParse::Command(&Str, TEXT("RecreateLandscapeCollision")))
 	{
 		// InWorld above is the PIE world if PIE is active, but this is specifically an editor command
diff --git a/Engine/Source/Editor/WorldBrowser/Private/Tiles/STiledLandscapeImportDlg.cpp b/Engine/Source/Editor/WorldBrowser/Private/Tiles/STiledLandscapeImportDlg.cpp
--- a/Engine/Source/Editor/WorldBrowser/Private/Tiles/STiledLandscapeImportDlg.cpp
+++ b/Engine/Source/Editor/WorldBrowser/Private/Tiles/STiledLandscapeImportDlg.cpp
@@ -9,13 +9,14 @@
 #include "LandscapeProxy.h"
 #include "SNumericEntryBox.h"
 #include "PropertyCustomizationHelpers.h"
-#include "ImageWrapper.h"
+#include "LandscapeEditorModule.h"
+#include "LandscapeDataAccess.h"
 
 #define LOCTEXT_NAMESPACE "WorldBrowser"
 
-static int32 CalcLandscapeSquareResolution(int32 ComponetsNumX, int32 SectionNumX, int32 SectionQuadsNumX)
+static int32 CalcLandscapeSquareResolution(int32 ComponentsNumX, int32 SectionNumX, int32 SectionQuadsNumX)
 {
-	return ComponetsNumX*SectionNumX*SectionQuadsNumX+1;
+	return ComponentsNumX * SectionNumX * SectionQuadsNumX + 1;
 }
 
 /**
@@ -252,9 +253,8 @@ void STiledLandcapeImportDlg::Construct(const FArguments& InArgs, TSharedPtr<SWi
 		]
 	];
 
-
 	GenerateAllPossibleTileConfigurations();
-	SetPossibleConfigurationsForFileSize(0);
+	ImportSettings.ComponentsNum = 0;
 }
 
 TSharedRef<SWidget> STiledLandcapeImportDlg::HandleTileConfigurationComboBoxGenarateWidget(TSharedPtr<FTileImportConfiguration> InItem) const
@@ -334,7 +334,7 @@ TSharedRef<ITableRow> STiledLandcapeImportDlg::OnGenerateWidgetForLayerDataListV
 						SNew(SButton)
 						.HAlign(HAlign_Center)
 						.ContentPadding(FEditorStyle::GetMargin("StandardDialog.ContentPadding"))
-						.OnClicked(this, &STiledLandcapeImportDlg::OnClickedSelectWeighmapTiles, InLayerData)
+						.OnClicked(this, &STiledLandcapeImportDlg::OnClickedSelectWeightmapTiles, InLayerData)
 						.Text(LOCTEXT("TiledLandscapeImport_SelectWeightmapButtonText", "Select Weightmap Tiles..."))
 					]
 				]
@@ -431,9 +431,11 @@ FReply STiledLandcapeImportDlg::OnClickedSelectHeightmapTiles()
 	TotalLandscapeRect = FIntRect(MAX_int32, MAX_int32, MIN_int32, MIN_int32);
 	ImportSettings.HeightmapFileList.Empty();
 	ImportSettings.TileCoordinates.Empty();
-	
-	SetPossibleConfigurationsForFileSize(0);
-	
+
+	ActiveConfigurations.Empty();
+	ImportSettings.ComponentsNum = 0;
+	StatusMessage = FText();
+
 	IDesktopPlatform* DesktopPlatform = FDesktopPlatformModule::Get();
 	if (DesktopPlatform)
 	{
@@ -441,177 +443,122 @@ FReply STiledLandcapeImportDlg::OnClickedSelectHeightmapTiles()
 		{
 			void* ParentWindowWindowHandle = ParentWindow->GetNativeWindow()->GetOSWindowHandle();
 
-			const FString ImportFileTypes = TEXT("All Heightmap tiles|*.png;*.raw;*.r16|Heightmap png tiles|*.png|Heightmap .raw/.r16 tiles|*.raw;*.r16|All files (*.*)|*.*");
+			ILandscapeEditorModule& LandscapeEditorModule = FModuleManager::GetModuleChecked<ILandscapeEditorModule>("LandscapeEditor");
+			const TCHAR* FileTypes = LandscapeEditorModule.GetHeightmapImportDialogTypeString();
+
 			bool bOpened = DesktopPlatform->OpenFileDialog(
 								ParentWindowWindowHandle,
 								LOCTEXT("SelectHeightmapTiles", "Select heightmap tiles").ToString(),
 								*FEditorDirectories::Get().GetLastDirectory(ELastDirectory::UNR),
 								TEXT(""),
-								*ImportFileTypes,
+								FileTypes,
 								EFileDialogFlags::Multiple,
 								ImportSettings.HeightmapFileList);
 
-			if (bOpened && ImportSettings.HeightmapFileList.Num())
+			if (bOpened && ImportSettings.HeightmapFileList.Num() > 0)
 			{
 				IFileManager& FileManager = IFileManager::Get();
 				bool bValidTiles = true;
+				int32 TargetSizeX = 0;
 
-				// All heightmap tiles have to be the same size and have correct tile position encoded into filename
 				const FString TargetExtension = FPaths::GetExtension(ImportSettings.HeightmapFileList[0], true);
-				const int64 TargetRawFileSize = TargetExtension != ".png" ? FileManager.FileSize(*ImportSettings.HeightmapFileList[0]) : 0;
-				int32 TargetSizeX = 0;
+				const ILandscapeHeightmapFileFormat* HeightmapFormat = LandscapeEditorModule.GetHeightmapFormatByExtension(*TargetExtension);
+
+				// All heightmap tiles have to be the same size and have correct tile position encoded into filename
 				for (const FString& Filename : ImportSettings.HeightmapFileList)
 				{
 					FFormatNamedArguments Arguments;
 					Arguments.Add(TEXT("FileName"), FText::FromString(Filename));
 
-					if (!Filename.EndsWith(TargetExtension))
+					FIntPoint TileCoordinate = ExtractTileCoordinates(FPaths::GetBaseFilename(Filename));
+					if (TileCoordinate.GetMin() < 0)
 					{
-						StatusMessage = FText::Format(LOCTEXT("TiledLandscapeImport_HeightmapMixedFileTypes", "File ({FileName}) has a different file extension, please use all the same type (16-bit grayscale png preferred)."), Arguments);
+						StatusMessage = FText::Format(LOCTEXT("TiledLandscapeImport_HeightmapTileInvalidName", "File name ({FileName}) should match pattern: <name>_X<number>_Y<number>."), Arguments);
 						bValidTiles = false;
 						break;
 					}
 
-					if (TargetExtension == ".raw" || TargetExtension == ".r16")
+					if (!Filename.EndsWith(TargetExtension))
 					{
-						int64 FileSize = FileManager.FileSize(*Filename);
-						if ((FileSize & 1) != 0) // we expect 2 bytes per sample
-						{
-							StatusMessage = FText::Format(LOCTEXT("TiledLandscapeImport_HeightmapTileInvalidFormat", "File ({FileName}) should have raw data with 2 bytes per sample."), Arguments);
-							bValidTiles = false;
-							break;
-						}
-
-						if (FileSize != TargetRawFileSize)
-						{
-							Arguments.Add(TEXT("FileSize"), FileSize);
-							Arguments.Add(TEXT("TargetFileSize"), TargetRawFileSize);
-							StatusMessage = FText::Format(LOCTEXT("TiledLandscapeImport_HeightmapTileSizeMismatch", "File ({FileName}) size ({FileSize}) should match other tiles file size ({TargetFileSize})."), Arguments);
-							bValidTiles = false;
-							break;
-						}
-
-						if (FMath::Square((int64)FMath::Sqrt(FileSize/2)) != FileSize/2)
-						{
-							StatusMessage = FText::Format(LOCTEXT("TiledLandscapeImport_NotSquare", "File ({FileName}) is not square."), Arguments);
-							bValidTiles = false;
-							break;
-						}
+						StatusMessage = FText::Format(LOCTEXT("TiledLandscapeImport_HeightmapMixedFileTypes", "File ({FileName}) has a different file extension, please use all the same type (16-bit grayscale png preferred)."), Arguments);
+						bValidTiles = false;
+						break;
 					}
-					else if (TargetExtension == ".png")
+
+					if (!HeightmapFormat)
 					{
-						TArray<uint8> ImportData;
-						if (!FFileHelper::LoadFileToArray(ImportData, *Filename, FILEREAD_Silent))
-						{
-							StatusMessage = FText::Format(LOCTEXT("TiledLandscapeImport_ErrorLoadingFile", "Error loading file ({FileName})."), Arguments);
-							bValidTiles = false;
-							break;
-						}
-						else
-						{
-							IImageWrapperModule& ImageWrapperModule = FModuleManager::LoadModuleChecked<IImageWrapperModule>("ImageWrapper");
-							IImageWrapperPtr ImageWrapper = ImageWrapperModule.CreateImageWrapper(EImageFormat::PNG);
-
-							if (!ImageWrapper->SetCompressed(ImportData.GetData(), ImportData.Num()))
-							{
-								StatusMessage = FText::Format(LOCTEXT("TiledLandscapeImport_ErrorLoadingFile", "Error loading file ({FileName})."), Arguments);
-								bValidTiles = false;
-								break;
-							}
-							else if (ImageWrapper->GetWidth() <= 0 || ImageWrapper->GetHeight() <= 0)
-							{
-								StatusMessage = FText::Format(LOCTEXT("TiledLandscapeImport_ErrorLoadingFile", "Error loading file ({FileName})."), Arguments);
-								bValidTiles = false;
-								break;
-							}
-							else if (ImageWrapper->GetWidth() != ImageWrapper->GetHeight())
-							{
-								StatusMessage = FText::Format(LOCTEXT("TiledLandscapeImport_NotSquare", "File ({FileName}) is not square."), Arguments);
-								bValidTiles = false;
-								break;
-							}
-							else if (ImageWrapper->GetFormat() != ERGBFormat::Gray)
-							{
-								StatusMessage = FText::Format(LOCTEXT("TiledLandscapeImport_HeightmapFileColorPng", "File ({FileName}) appears to be a color png, grayscale is expected. The import *can* continue, but the result may be lower quality than desired."), Arguments);
-								bValidTiles = false;
-								break;
-							}
-							else if (ImageWrapper->GetBitDepth() != 16)
-							{
-								StatusMessage = FText::Format(LOCTEXT("TiledLandscapeImport_HeightmapFileLowBitDepth", "File ({FileName}) appears to be an 8-bit png, 16-bit is preferred. The import *can* continue, but the result may be lower quality than desired."), Arguments);
-								bValidTiles = false;
-								break;
-							}
-
-							if (TargetSizeX == 0)
-							{
-								TargetSizeX = ImageWrapper->GetWidth();
-							}
-							else if (ImageWrapper->GetWidth() != TargetSizeX)
-							{
-								Arguments.Add(TEXT("Size"), ImageWrapper->GetWidth());
-								Arguments.Add(TEXT("TargetSize"), TargetSizeX);
-								StatusMessage = FText::Format(LOCTEXT("TiledLandscapeImport_HeightmapPngTileSizeMismatch", "File ({FileName}) size ({Size}\u00D7{Size}) should match other tiles file size ({TargetSize}\u00D7{TargetSize})."), Arguments);
-								bValidTiles = false;
-								break;
-							}
-						}
+						StatusMessage = FText::Format(LOCTEXT("TiledLandscapeImport_UnrecognisedExtension", "Error loading file ({FileName}), unrecognised extension."), Arguments);
+						bValidTiles = false;
+						break;
 					}
-					else
+
+					FLandscapeHeightmapInfo HeightmapInfo = HeightmapFormat->Validate(*Filename);
+					if (HeightmapInfo.ResultCode != ELandscapeImportResult::Success)
 					{
-						StatusMessage = FText::Format(LOCTEXT("TiledLandscapeImport_UnrecognisedExtension", "Error loading file ({FileName}), unrecognised extension."), Arguments);
+						StatusMessage = HeightmapInfo.ErrorMessage;
 						bValidTiles = false;
 						break;
 					}
 
-					FIntPoint TileCoordinate = ExtractTileCoordinates(FPaths::GetBaseFilename(Filename));
-					if (TileCoordinate.GetMin() < 0)
+					const FLandscapeFileResolution* FoundSquareResolution =
+						HeightmapInfo.PossibleResolutions.FindByPredicate(
+							[](const FLandscapeFileResolution& Resolution)
 					{
-						StatusMessage = FText::Format(LOCTEXT("TiledLandscapeImport_HeightmapTileInvalidName", "File name ({FileName}) should match pattern: <name>_X<number>_Y<number>."), Arguments);
+						return Resolution.Width == Resolution.Height;
+					});
+					if (!FoundSquareResolution)
+					{
+						StatusMessage = FText::Format(LOCTEXT("TiledLandscapeImport_NotSquare", "File ({FileName}) is not square."), Arguments);
 						bValidTiles = false;
 						break;
 					}
 
-					TotalLandscapeRect.Include(TileCoordinate);
-					ImportSettings.TileCoordinates.Add(TileCoordinate);
-				}
-				
-				if (bValidTiles)
-				{
-					if (TargetExtension == ".raw" || TargetExtension == ".r16")
+					if (TargetSizeX == 0)
 					{
-						if (SetPossibleConfigurationsForFileSize(TargetRawFileSize) < 1)
+						TargetSizeX = FoundSquareResolution->Width;
+						if (HeightmapInfo.DataScale.IsSet())
 						{
-							int64 NumSamples = TargetRawFileSize / 2; // 2 bytes per sample
-							FFormatNamedArguments Arguments;
-							Arguments.Add(TEXT("NumSamples"), NumSamples);
-							StatusMessage = FText::Format(LOCTEXT("TiledLandscapeImport_HeightmapTileInvalidSize", "No landscape configuration found for ({NumSamples}) samples."), Arguments);
-							bValidTiles = false;
+							ImportSettings.Scale3D = HeightmapInfo.DataScale.GetValue();
+							ImportSettings.Scale3D.Z *= LANDSCAPE_INV_ZSCALE;
 						}
 					}
-					else if (TargetExtension == ".png")
+					else
 					{
-						if (SetPossibleConfigurationsForFileWidth(TargetSizeX) < 1)
+						if (TargetSizeX != FoundSquareResolution->Width)
 						{
-							FFormatNamedArguments Arguments;
-							Arguments.Add(TEXT("Size"), TargetSizeX);
-							StatusMessage = FText::Format(LOCTEXT("TiledLandscapeImport_HeightmapPngTileInvalidSize", "No landscape configuration found for ({Size}\u00D7{Size})."), Arguments);
+							Arguments.Add(TEXT("Size"), FoundSquareResolution->Width);
+							Arguments.Add(TEXT("TargetSize"), TargetSizeX);
+							StatusMessage = FText::Format(LOCTEXT("TiledLandscapeImport_HeightmapPngTileSizeMismatch", "File ({FileName}) size ({Size}\u00D7{Size}) should match other tiles file size ({TargetSize}\u00D7{TargetSize})."), Arguments);
 							bValidTiles = false;
+							break;
 						}
 					}
-					else
+
+					TotalLandscapeRect.Include(TileCoordinate);
+					ImportSettings.TileCoordinates.Add(TileCoordinate);
+				}
+
+				if (bValidTiles)
+				{
+					if (SetPossibleConfigurationsForFileWidth(TargetSizeX) < 1)
 					{
-						check(0);
+						FFormatNamedArguments Arguments;
+						Arguments.Add(TEXT("Size"), TargetSizeX);
+						StatusMessage = FText::Format(LOCTEXT("TiledLandscapeImport_HeightmapPngTileInvalidSize", "No landscape configuration found for ({Size}\u00D7{Size})."), Arguments);
+						bValidTiles = false;
 					}
 				}
 			}
 		}
 	}
-		
+
+	// Refresh combo box with active configurations
+	TileConfigurationComboBox->RefreshOptions();
+
 	return FReply::Handled();
 }
 
-FReply STiledLandcapeImportDlg::OnClickedSelectWeighmapTiles(TSharedPtr<FTiledLandscapeImportSettings::LandscapeLayerSettings> InLayerData)
+FReply STiledLandcapeImportDlg::OnClickedSelectWeightmapTiles(TSharedPtr<FTiledLandscapeImportSettings::LandscapeLayerSettings> InLayerData)
 {
 	InLayerData->WeightmapFiles.Empty();
 
@@ -624,13 +571,15 @@ FReply STiledLandcapeImportDlg::OnClickedSelectWeighmapTiles(TSharedPtr<FTiledLa
 
 			TArray<FString> WeightmapFilesList;
 
-			const FString ImportFileTypes = TEXT("All Weightmap tiles|*.png;*.raw;*.r8|Weightmap png tiles|*.png|Weightmap .raw/.r8 files|*.raw;*.r8|All files (*.*)|*.*");
+			ILandscapeEditorModule& LandscapeEditorModule = FModuleManager::GetModuleChecked<ILandscapeEditorModule>("LandscapeEditor");
+			const TCHAR* FileTypes = LandscapeEditorModule.GetWeightmapImportDialogTypeString();
+
 			bool bOpened = DesktopPlatform->OpenFileDialog(
 								ParentWindowWindowHandle,
 								LOCTEXT("SelectWeightmapTiles", "Select weightmap tiles").ToString(),
 								*FEditorDirectories::Get().GetLastDirectory(ELastDirectory::UNR),
 								TEXT(""),
-								*ImportFileTypes,
+								FileTypes,
 								EFileDialogFlags::Multiple,
 								WeightmapFilesList);
 
@@ -664,18 +613,18 @@ FReply STiledLandcapeImportDlg::OnClickedImport()
 	{
 		ImportSettings.LandscapeLayerSettingsList.Add(*(LayerDataList[LayerIdx].Get()));
 	}
-					
+
 	ParentWindow->RequestDestroyWindow();
-	bShouldImport = true;	
-	return FReply::Handled();	
+	bShouldImport = true;
+	return FReply::Handled();
 }
 
 FReply STiledLandcapeImportDlg::OnClickedCancel()
 {
 	ParentWindow->RequestDestroyWindow();
-	bShouldImport = false;	
+	bShouldImport = false;
 
-	return FReply::Handled();	
+	return FReply::Handled();
 }
 
 FString STiledLandcapeImportDlg::GetLandscapeMaterialPath() const
@@ -691,34 +640,6 @@ void STiledLandcapeImportDlg::OnLandscapeMaterialChanged(const FAssetData& Asset
 	UpdateLandscapeLayerList();
 }
 
-int32 STiledLandcapeImportDlg::SetPossibleConfigurationsForFileSize(int64 TargetFileSize)
-{
-	int32 Idx = AllConfigurations.IndexOfByPredicate([&](const FTileImportConfiguration& A){
-		return TargetFileSize == A.ImportFileSize;
-	});
-
-	ActiveConfigurations.Empty();
-	ImportSettings.ComponentsNum = 0; // Set invalid options
-	StatusMessage = FText();
-
-	// AllConfigurations - sorted by resolution
-	while (AllConfigurations.IsValidIndex(Idx) && AllConfigurations[Idx].ImportFileSize == TargetFileSize)
-	{
-		TSharedPtr<FTileImportConfiguration> TileConfig = MakeShareable(new FTileImportConfiguration(AllConfigurations[Idx++]));
-		ActiveConfigurations.Add(TileConfig);
-	}
-
-	// Refresh combo box with active configurations
-	TileConfigurationComboBox->RefreshOptions();
-	// Set first configuration as active
-	if (ActiveConfigurations.Num())
-	{
-		TileConfigurationComboBox->SetSelectedItem(ActiveConfigurations[0]);
-	}
-
-	return ActiveConfigurations.Num();
-}
-
 int32 STiledLandcapeImportDlg::SetPossibleConfigurationsForFileWidth(int64 TargetFileWidth)
 {
 	int32 Idx = AllConfigurations.IndexOfByPredicate([&](const FTileImportConfiguration& A){
@@ -727,7 +648,6 @@ int32 STiledLandcapeImportDlg::SetPossibleConfigurationsForFileWidth(int64 Targe
 
 	ActiveConfigurations.Empty();
 	ImportSettings.ComponentsNum = 0; // Set invalid options
-	StatusMessage = FText();
 
 	// AllConfigurations - sorted by resolution
 	for (; AllConfigurations.IsValidIndex(Idx) && AllConfigurations[Idx].SizeX == TargetFileWidth; ++Idx)
@@ -761,14 +681,12 @@ void STiledLandcapeImportDlg::GenerateAllPossibleTileConfigurations()
 				Entry.NumSectionsPerComponent	= SectionsPerComponent;
 				Entry.NumQuadsPerSection		= (1 << QuadsPerSection) - 1;
 				Entry.SizeX = CalcLandscapeSquareResolution(Entry.NumComponents, Entry.NumSectionsPerComponent, Entry.NumQuadsPerSection);
-				// Calculate expected file size for this setup
-				Entry.ImportFileSize = FMath::Square(Entry.SizeX)*2; // 2 bytes per sample 
 
 				AllConfigurations.Add(Entry);
 			}
 		}
 	}
-	
+
 	// Sort by resolution
 	AllConfigurations.Sort([](const FTileImportConfiguration& A, const FTileImportConfiguration& B){
 		if (A.SizeX == B.SizeX)
diff --git a/Engine/Source/Editor/WorldBrowser/Private/Tiles/STiledLandscapeImportDlg.h b/Engine/Source/Editor/WorldBrowser/Private/Tiles/STiledLandscapeImportDlg.h
--- a/Engine/Source/Editor/WorldBrowser/Private/Tiles/STiledLandscapeImportDlg.h
+++ b/Engine/Source/Editor/WorldBrowser/Private/Tiles/STiledLandscapeImportDlg.h
@@ -7,12 +7,10 @@
 class STiledLandcapeImportDlg
 	: public SCompoundWidget
 {
-public:	
-	/** */		
+public:
+	/** */
 	struct FTileImportConfiguration
 	{
-		int64 ImportFileSize;
-		
 		int32 SizeX;
 		int32 NumComponents;
 		int32 NumSectionsPerComponent;
@@ -22,7 +20,7 @@ public:
 	SLATE_BEGIN_ARGS( STiledLandcapeImportDlg )
 		{}
 	SLATE_END_ARGS()
-	
+
 	void Construct(const FArguments& InArgs, TSharedPtr<SWindow> InParentWindow);
 
 	/** */
@@ -33,7 +31,7 @@ private:
 	/** */
 	TSharedRef<SWidget> HandleTileConfigurationComboBoxGenarateWidget(TSharedPtr<FTileImportConfiguration> InItem) const;
 	FText				GetTileConfigurationText() const;
-	
+
 	/** */
 	TSharedRef<ITableRow> OnGenerateWidgetForLayerDataListView(TSharedPtr<FTiledLandscapeImportSettings::LandscapeLayerSettings> InLayerData, const TSharedRef<STableViewBase>& OwnerTable);
 
@@ -52,16 +50,16 @@ private:
 	/** */
 	ECheckBoxState GetFlipYAxisState() const;
 	void OnFlipYAxisStateChanged(ECheckBoxState NewState);
-		
+
 	/** */
 	void OnSetImportConfiguration(TSharedPtr<FTileImportConfiguration> InTileConfig, ESelectInfo::Type SelectInfo);
 
 	/** */
 	bool IsImportEnabled() const;
-		
+
 	/** */
 	FReply OnClickedSelectHeightmapTiles();
-	FReply OnClickedSelectWeighmapTiles(TSharedPtr<FTiledLandscapeImportSettings::LandscapeLayerSettings> InLayerData);
+	FReply OnClickedSelectWeightmapTiles(TSharedPtr<FTiledLandscapeImportSettings::LandscapeLayerSettings> InLayerData);
 	FReply OnClickedImport();
 	FReply OnClickedCancel();
 
@@ -78,7 +76,6 @@ private:
 	void OnLayerBlendStateChanged(ECheckBoxState NewState, TSharedPtr<FTiledLandscapeImportSettings::LandscapeLayerSettings> InLayerData);
 
 	/** */
-	int32 SetPossibleConfigurationsForFileSize(int64 InFileSize);
 	int32 SetPossibleConfigurationsForFileWidth(int64 TargetFileWidth);
 
 	/** */
@@ -96,22 +93,22 @@ private:
 
 	/** */
 	mutable FText StatusMessage;
-	
+
 	/** */
 	TSharedPtr<SWindow> ParentWindow;
-		
+
 	/** */
 	TSharedPtr<SComboBox<TSharedPtr<FTileImportConfiguration>>> TileConfigurationComboBox;
-	
+
 	TSharedPtr<SListView<TSharedPtr<FTiledLandscapeImportSettings::LandscapeLayerSettings>>>	LayerDataListView;
-	TArray<TSharedPtr<FTiledLandscapeImportSettings::LandscapeLayerSettings>>					LayerDataList;		
-		
+	TArray<TSharedPtr<FTiledLandscapeImportSettings::LandscapeLayerSettings>>					LayerDataList;
+
 	/** */
 	FTiledLandscapeImportSettings ImportSettings;
-	
+
 	/** */
-	FIntRect TotalLandscapeRect;	
-	
+	FIntRect TotalLandscapeRect;
+
 	TArray<FTileImportConfiguration> AllConfigurations;
 	TArray<TSharedPtr<FTileImportConfiguration>> ActiveConfigurations;
 };
diff --git a/Engine/Source/Editor/WorldBrowser/Private/Tiles/WorldTileCollectionModel.cpp b/Engine/Source/Editor/WorldBrowser/Private/Tiles/WorldTileCollectionModel.cpp
--- a/Engine/Source/Editor/WorldBrowser/Private/Tiles/WorldTileCollectionModel.cpp
+++ b/Engine/Source/Editor/WorldBrowser/Private/Tiles/WorldTileCollectionModel.cpp
@@ -11,7 +11,6 @@
 #include "MaterialUtilities.h"
 #include "RawMesh.h"
 #include "LandscapeEditorUtils.h"
-#include "ImageWrapper.h"
 
 #include "WorldTileDetails.h"
 #include "WorldTileDetailsCustomization.h"
@@ -28,6 +27,7 @@
 #include "LandscapeMeshProxyComponent.h"
 #include "LandscapeLayerInfoObject.h"
 #include "Components/HierarchicalInstancedStaticMeshComponent.h"
+#include "LandscapeEditorModule.h"
 
 #define LOCTEXT_NAMESPACE "WorldBrowser"
 
@@ -1430,72 +1430,45 @@ void FWorldTileCollectionModel::AddLandscapeProxy_Executed(FWorldTileModel::EWor
 	}
 }
 
-template<typename DataType>
-static bool ReadRawFile(TArray<DataType>& Result, const FString& Filename, uint32 Flags = 0)
-{
-	auto Reader = TScopedPointer<FArchive>(IFileManager::Get().CreateFileReader(*Filename, Flags));
-	if (!Reader)
-	{
-		if (!(Flags & FILEREAD_Silent))
-		{
-			UE_LOG(LogStreaming, Warning, TEXT("Failed to read file '%s' error."), *Filename);
-		}
-		return 0;
-	}
-	Result.SetNumUninitialized(Reader->TotalSize()/Result.GetTypeSize());
-	Reader->Serialize(Result.GetData(), Result.Num()*Result.GetTypeSize());
-	return Reader->Close();
-}
-
-static bool ReadHeightmapFile(TArray<uint16>& Result, const FString& Filename, uint32 Flags = 0)
+static bool ReadHeightmapFile(TArray<uint16>& Result, const FString& Filename, int32 ExpectedWidth, int32 ExpectedHeight)
 {
 	bool bResult = true;
-	if (Filename.EndsWith(TEXT(".png")))
+
+	ILandscapeEditorModule& LandscapeEditorModule = FModuleManager::GetModuleChecked<ILandscapeEditorModule>("LandscapeEditor");
+	const ILandscapeHeightmapFileFormat* HeightmapFormat = LandscapeEditorModule.GetHeightmapFormatByExtension(*FPaths::GetExtension(Filename, true));
+
+	FLandscapeHeightmapImportData ImportData = HeightmapFormat->Import(*Filename, {(uint32)ExpectedWidth, (uint32)ExpectedHeight});
+	if (ImportData.ResultCode != ELandscapeImportResult::Error)
 	{
-		TArray<uint8> FileRawData;
-		bResult = ReadRawFile(FileRawData, Filename, Flags);
-		if (bResult)
-		{
-			IImageWrapperModule& ImageWrapperModule = FModuleManager::LoadModuleChecked<IImageWrapperModule>("ImageWrapper");
-			IImageWrapperPtr ImageWrapper = ImageWrapperModule.CreateImageWrapper(EImageFormat::PNG);
-			bResult = ImageWrapper->SetCompressed(FileRawData.GetData(), FileRawData.Num());
-			if (bResult)
-			{
-				const TArray<uint8>* RawData = nullptr;
-				bResult = ImageWrapper->GetRaw(ERGBFormat::Gray, 16, RawData);
-				if (bResult)
-				{
-					Result.AddUninitialized(RawData->Num()/2);
-					FMemory::Memcpy(Result.GetData(), RawData->GetData(), RawData->Num());
-				}
-			}
-		}
+		Result = MoveTemp(ImportData.Data);
 	}
 	else
 	{
-		bResult = ReadRawFile(Result, Filename, Flags);
+		UE_LOG(LogStreaming, Warning, TEXT("%s"), *ImportData.ErrorMessage.ToString());
+		Result.Empty();
+		bResult = false;
 	}
 
 	return bResult;
 }
 
-static bool ReadWeightmapFile(TArray<uint8>& Result, const FString& Filename, uint32 Flags = 0)
+static bool ReadWeightmapFile(TArray<uint8>& Result, const FString& Filename, FName LayerName, int32 ExpectedWidth, int32 ExpectedHeight)
 {
-	bool bResult = ReadRawFile(Result, Filename, Flags);
-	if (bResult && Filename.EndsWith(TEXT(".png")))
+	bool bResult = true;
+
+	ILandscapeEditorModule& LandscapeEditorModule = FModuleManager::GetModuleChecked<ILandscapeEditorModule>("LandscapeEditor");
+	const ILandscapeWeightmapFileFormat* WeightmapFormat = LandscapeEditorModule.GetWeightmapFormatByExtension(*FPaths::GetExtension(Filename, true));
+
+	FLandscapeWeightmapImportData ImportData = WeightmapFormat->Import(*Filename, LayerName, {(uint32)ExpectedWidth, (uint32)ExpectedHeight});
+	if (ImportData.ResultCode != ELandscapeImportResult::Error)
 	{
-		IImageWrapperModule& ImageWrapperModule = FModuleManager::LoadModuleChecked<IImageWrapperModule>("ImageWrapper");
-		IImageWrapperPtr ImageWrapper = ImageWrapperModule.CreateImageWrapper(EImageFormat::PNG);
-		bResult = ImageWrapper->SetCompressed(Result.GetData(), Result.Num());
-		if (bResult)
-		{
-			const TArray<uint8>* RawData = nullptr;
-			bResult = ImageWrapper->GetRaw(ERGBFormat::Gray, 8, RawData);
-			if (bResult)
-			{
-				Result = *RawData;
-			}
-		}
+		Result = MoveTemp(ImportData.Data);
+	}
+	else
+	{
+		UE_LOG(LogStreaming, Warning, TEXT("%s"), *ImportData.ErrorMessage.ToString());
+		Result.Empty();
+		bResult = false;
 	}
 
 	return bResult;
@@ -1551,7 +1524,7 @@ static void SetupLandscapeImportLayers(const FTiledLandscapeImportSettings& InIm
 			if (WeightmapFile)
 			{
 				LayerImportInfo.SourceFilePath = *WeightmapFile;
-				ReadWeightmapFile(LayerImportInfo.LayerData, LayerImportInfo.SourceFilePath, FILEREAD_Silent);
+				ReadWeightmapFile(LayerImportInfo.LayerData, LayerImportInfo.SourceFilePath, LayerImportInfo.LayerName, InImportSettings.SizeX, InImportSettings.SizeX);
 			}
 		}
 
@@ -1616,7 +1589,7 @@ void FWorldTileCollectionModel::ImportTiledLandscape_Executed()
 			{
 				Landscape->EditorLayerSettings.Add(FLandscapeEditorLayerSettings(ImportLayerInfo.LayerInfo));
 			}
-			ULandscapeInfo::RecreateLandscapeInfo(GetWorld(), false);
+			Landscape->CreateLandscapeInfo();
 		}
 
 		// Import tiles
@@ -1645,7 +1618,7 @@ void FWorldTileCollectionModel::ImportTiledLandscape_Executed()
 			SetupLandscapeImportLayers(ImportSettings, GetWorld()->GetOutermost()->GetName(), TileIndex, TileImportSettings.ImportLayers);
 			TileImportSettings.ImportLayerType = ELandscapeImportAlphamapType::Additive;
 
-			if (ReadHeightmapFile(TileImportSettings.HeightData, Filename, FILEREAD_Silent))
+			if (ReadHeightmapFile(TileImportSettings.HeightData, Filename, TileImportSettings.SizeX, TileImportSettings.SizeY))
 			{
 				FString MapFileName = WorldRootPath + TileName + FPackageName::GetMapPackageExtension();
 				// Create a new world - so we can 'borrow' its level
@@ -1738,6 +1711,7 @@ void FWorldTileCollectionModel::ReimportTiledLandscape_Executed(FName TargetLaye
 		TileModel->SetVisible(true);
 
 		ALandscapeProxy* Landscape = TileModel->GetLandscape();
+		FIntRect LandscapeSize = Landscape->GetBoundingRect();
 
 		ULandscapeLayerInfoObject* DataLayer = ALandscapeProxy::VisibilityLayer;
 
@@ -1746,7 +1720,7 @@ void FWorldTileCollectionModel::ReimportTiledLandscape_Executed(FName TargetLaye
 			if (!Landscape->ReimportHeightmapFilePath.IsEmpty())
 			{
 				TArray<uint16> RawData;
-				ReadHeightmapFile(RawData, *Landscape->ReimportHeightmapFilePath, FILEREAD_Silent);
+				ReadHeightmapFile(RawData, *Landscape->ReimportHeightmapFilePath, LandscapeSize.Width(), LandscapeSize.Height());
 				LandscapeEditorUtils::SetHeightmapData(Landscape, RawData);
 			}
 		}
@@ -1759,7 +1733,7 @@ void FWorldTileCollectionModel::ReimportTiledLandscape_Executed(FName TargetLaye
 					if (!LayerSettings.ReimportLayerFilePath.IsEmpty())
 					{
 						TArray<uint8> RawData;
-						ReadWeightmapFile(RawData, *LayerSettings.ReimportLayerFilePath, FILEREAD_Silent);
+						ReadWeightmapFile(RawData, *LayerSettings.ReimportLayerFilePath, LayerSettings.LayerInfoObj->LayerName, LandscapeSize.Width(), LandscapeSize.Height());
 						LandscapeEditorUtils::SetWeightmapData(Landscape, LayerSettings.LayerInfoObj, RawData);
 
 						if (TargetLayer != NAME_None)
diff --git a/Engine/Source/Runtime/AnimGraphRuntime/Private/AnimNodes/AnimNode_RandomPlayer.cpp b/Engine/Source/Runtime/AnimGraphRuntime/Private/AnimNodes/AnimNode_RandomPlayer.cpp
--- a/Engine/Source/Runtime/AnimGraphRuntime/Private/AnimNodes/AnimNode_RandomPlayer.cpp
+++ b/Engine/Source/Runtime/AnimGraphRuntime/Private/AnimNodes/AnimNode_RandomPlayer.cpp
@@ -198,7 +198,7 @@ void FAnimNode_RandomPlayer::Evaluate(FPoseContext& Output)
 					CurrentSequence->GetAnimationPose(Poses[0], Curves[0], FAnimExtractContext(CurrentData.InternalTimeAccumulator, AnimProxy->ShouldExtractRootMotion()));
 					NextSequence->GetAnimationPose(Poses[1], Curves[1], FAnimExtractContext(NextData.InternalTimeAccumulator, AnimProxy->ShouldExtractRootMotion()));
 
-					FAnimationRuntime::BlendPosesTogether(TFixedSizeArrayView<FCompactPose>(Poses, 2), TFixedSizeArrayView<FBlendedCurve>(Curves, 2), TFixedSizeArrayView<float>(Weights, 2), Output.Pose, Output.Curve);
+					FAnimationRuntime::BlendPosesTogether(Poses, Curves, Weights, Output.Pose, Output.Curve);
 				}
 				else
 				{
diff --git a/Engine/Source/Runtime/Core/Private/Containers/Algo/AlgosTest.cpp b/Engine/Source/Runtime/Core/Private/Containers/Algo/AlgosTest.cpp
new file mode 100644
--- /dev/null
+++ b/Engine/Source/Runtime/Core/Private/Containers/Algo/AlgosTest.cpp
@@ -0,0 +1,144 @@
+// Copyright 1998-2016 Epic Games, Inc. All Rights Reserved.
+
+#include "CorePrivatePCH.h"
+#include "AutomationTest.h"
+#include "Algo/Copy.h"
+#include "Algo/Transform.h"
+
+IMPLEMENT_SIMPLE_AUTOMATION_TEST(FAlgosTest, "System.Core.Misc.Algos", EAutomationTestFlags::ApplicationContextMask | EAutomationTestFlags::SmokeFilter)
+
+#define NUM_TEST_OBJECTS      32
+
+
+
+bool FAlgosTest::RunTest(const FString& Parameters)
+{
+	// setup
+	TArray<int> TestData;
+	for (int i = 0; i < NUM_TEST_OBJECTS; ++i)
+	{
+		TestData.Add(i);
+	}
+	TArray<int> TestData2;
+	for (int i = 0; i < NUM_TEST_OBJECTS; ++i)
+	{
+		TestData2.Add(FMath::Rand());
+	}
+
+	// copy
+	{
+		TArray<int> TestArray;
+		// empty array
+		Algo::Copy(TestData, TestArray);
+		check(TestArray == TestData);
+		// existing data
+		Algo::Copy(TestData2, TestArray);
+		check(TestArray.Num() == NUM_TEST_OBJECTS * 2);
+		for (int i = 0; i < NUM_TEST_OBJECTS; ++i)
+		{
+			check(TestArray[i] == TestData[i]);
+		}
+		for (int i = 0; i < NUM_TEST_OBJECTS; ++i)
+		{
+			check(TestArray[i + NUM_TEST_OBJECTS] == TestData2[i]);
+		}
+	}
+
+	// copy if
+	{
+		TArray<int> TestArray;
+		// empty array
+		Algo::CopyIf(TestData, TestArray, [](int i) { return (i % 2) == 0; });
+		int j = 0;
+		for (int i = 0; i < NUM_TEST_OBJECTS; ++i)
+		{
+			if (TestData[i] % 2 == 0)
+			{
+				check(TestArray[j] == TestData[i]);
+				++j;
+			}
+		}
+		// existing data
+		Algo::CopyIf(TestData2, TestArray, [](int i) { return (i % 2) == 0; });
+		j = 0;
+		for (int i = 0; i < NUM_TEST_OBJECTS; ++i)
+		{
+			if (TestData[i] % 2 == 0)
+			{
+				check(TestArray[j] == TestData[i]);
+				++j;
+			}
+		}
+		for (int i = 0; i < NUM_TEST_OBJECTS; ++i)
+		{
+			if (TestData2[i] % 2 == 0)
+			{
+				check(TestArray[j] == TestData2[i]);
+				++j;
+			}
+		}
+		check(j == TestArray.Num())
+	}
+
+	// transform
+	{
+		TArray<float> TestArray;
+		// empty array
+		Algo::Transform(TestData, TestArray, [](int i) { return FMath::DegreesToRadians(i); });
+		check(TestArray.Num() == NUM_TEST_OBJECTS);
+		for (int i = 0; i < TestArray.Num(); ++i)
+		{
+			check(TestArray[i] == FMath::DegreesToRadians(TestData[i]));
+		}
+		// existing data
+		Algo::Transform(TestData2, TestArray, [](int i) { return FMath::DegreesToRadians(i); });
+		check(TestArray.Num() == NUM_TEST_OBJECTS * 2);
+		for (int i = 0; i < NUM_TEST_OBJECTS; ++i)
+		{
+			check(TestArray[i] == FMath::DegreesToRadians(TestData[i]));
+		}
+		for (int i = 0; i < NUM_TEST_OBJECTS; ++i)
+		{
+			check(TestArray[i + NUM_TEST_OBJECTS] == FMath::DegreesToRadians(TestData2[i]));
+		}
+	}
+
+	// transform if
+	{
+		TArray<float> TestArray;
+		// empty array
+		Algo::TransformIf(TestData, TestArray, [](int i) { return (i % 2) == 0; }, [](int i) { return FMath::DegreesToRadians(i); });
+		int j = 0;
+		for (int i = 0; i < NUM_TEST_OBJECTS; ++i)
+		{
+			if (TestData[i] % 2 == 0)
+			{
+				check(TestArray[j] == FMath::DegreesToRadians(TestData[i]));
+				++j;
+			}
+		}
+
+		// existing data
+		Algo::TransformIf(TestData2, TestArray, [](int i) { return (i % 2) == 0; }, [](int i) { return FMath::DegreesToRadians(i); });
+		j = 0;
+		for (int i = 0; i < NUM_TEST_OBJECTS; ++i)
+		{
+			if (TestData[i] % 2 == 0)
+			{
+				check(TestArray[j] == FMath::DegreesToRadians(TestData[i]));
+				++j;
+			}
+		}
+		for (int i = 0; i < NUM_TEST_OBJECTS; ++i)
+		{
+			if (TestData2[i] % 2 == 0)
+			{
+				check(TestArray[j] == FMath::DegreesToRadians(TestData2[i]));
+				++j;
+			}
+		}
+		check(j == TestArray.Num())
+	}
+
+	return true;
+}
diff --git a/Engine/Source/Runtime/Core/Private/Containers/ContainersTest.cpp b/Engine/Source/Runtime/Core/Private/Containers/ContainersTest.cpp
--- a/Engine/Source/Runtime/Core/Private/Containers/ContainersTest.cpp
+++ b/Engine/Source/Runtime/Core/Private/Containers/ContainersTest.cpp
@@ -318,3 +318,68 @@ bool FContainersTest::RunTest( const FString& Parameters )
 
 	return true;
 }
+
+//Array view tests
+#include "Containers/ArrayView.h"
+
+namespace ArrayViewTests
+{
+	// commented out lines shouldn't compile
+
+	struct Base
+	{
+		int32 b;
+	};
+
+	struct Derived : public Base
+	{
+		int32 d;
+	};
+
+	template<typename T>
+	void TestFunction(TArrayView<T>)
+	{
+	}
+
+	bool RunTest()
+	{
+		// C array + derived-to-base conversions
+		Derived test1[13];
+		TestFunction<Derived>(test1);
+		//TestFunction<Base>(test1);
+		//TestFunction<const Base>(test1);
+
+		// C array of pointers + derived-to-base conversions
+		Derived* test2[13];
+		TestFunction<const Derived* const>(test2);
+		//TestFunction<const Derived*>(test2);
+		TestFunction<Derived* const>(test2);
+		//TestFunction<const Base* const>(test2);
+
+		// TArray + derived-to-base conversions
+		TArray<Derived> test3;
+		TestFunction<Derived>(test3);
+		//TestFunction<Base>(test3);
+		//TestFunction<const Base>(test3);
+
+		// const TArray
+		const TArray<Base> test4;
+		TestFunction<const Base>(test4);
+		//TestFunction<Base>(test4);
+
+		// TArray of const
+		TArray<const Base> test5;
+		TestFunction<const Base>(test5);
+		//TestFunction<Base>(test5);
+
+		// temporary C array
+		struct Test6
+		{
+			Base test[13];
+		};
+		TestFunction<const Base>(Test6().test);
+		//TestFunction<Base>(Test6().test); // shouldn't compile but VS allows it :(
+
+		return true;
+	}
+};
diff --git a/Engine/Source/Runtime/Core/Private/Misc/CoreMisc.cpp b/Engine/Source/Runtime/Core/Private/Misc/CoreMisc.cpp
--- a/Engine/Source/Runtime/Core/Private/Misc/CoreMisc.cpp
+++ b/Engine/Source/Runtime/Core/Private/Misc/CoreMisc.cpp
@@ -204,7 +204,7 @@ bool FFileHelper::LoadFileToString( FString& Result, const TCHAR* Filename, uint
 /**
  * Save a binary array to a file.
  */
-bool FFileHelper::SaveArrayToFile( const TArray<uint8>& Array, const TCHAR* Filename, IFileManager* FileManager /*= &IFileManager::Get()*/, uint32 WriteFlags )
+bool FFileHelper::SaveArrayToFile(TArrayView<const uint8> Array, const TCHAR* Filename, IFileManager* FileManager /*= &IFileManager::Get()*/, uint32 WriteFlags)
 {
 	FArchive* Ar = FileManager->CreateFileWriter( Filename, WriteFlags );
 	if( !Ar )
diff --git a/Engine/Source/Runtime/Core/Public/Containers/Algo/Accumulate.h b/Engine/Source/Runtime/Core/Public/Containers/Algo/Accumulate.h
--- a/Engine/Source/Runtime/Core/Public/Containers/Algo/Accumulate.h
+++ b/Engine/Source/Runtime/Core/Public/Containers/Algo/Accumulate.h
@@ -29,14 +29,14 @@ namespace Algo
 	* @return the result of summing all the elements of Input
 	*/
 	template <typename T, typename A, typename OpT>
-	FORCEINLINE T Accumulate(const A& Input, T init, OpT Op)
+	FORCEINLINE T Accumulate(const A& Input, T Init, OpT Op)
 	{
-		T result = init;
-		for (auto&& i : Input)
+		T Result = MoveTemp(Init);
+		for (auto&& InputElem : Input)
 		{
-			result = Op(result, i);
+			Result = Op(MoveTemp(Result), InputElem);
 		}
-		return result;
+		return Result;
 	}
 
 	/**
@@ -48,9 +48,9 @@ namespace Algo
 	 * @return the result of summing all the elements of Input
 	 */
 	template <typename T, typename A>
-	FORCEINLINE T Accumulate(const A& Input, T init)
+	FORCEINLINE T Accumulate(const A& Input, T Init)
 	{
-		return Accumulate(Input, init, TPlus<>());
+		return Accumulate(Input, MoveTemp(Init), TPlus<>());
 	}
 
 	/**
@@ -64,14 +64,14 @@ namespace Algo
 	* @return the result of mapping and then summing all the elements of Input
 	*/
 	template <typename T, typename A, typename MapT, typename OpT>
-	FORCEINLINE T TransformAccumulate(const A& Input, MapT MapOp, T init, OpT Op)
+	FORCEINLINE T TransformAccumulate(const A& Input, MapT MapOp, T Init, OpT Op)
 	{
-		T result = init;
-		for (auto&& i : Input)
+		T Result = MoveTemp(Init);
+		for (auto&& InputElem : Input)
 		{
-			result = Op(result, MapOp(i));
+			Result = Op(MoveTemp(Result), MapOp(InputElem));
 		}
-		return result;
+		return Result;
 	}
 
 	/**
@@ -84,8 +84,8 @@ namespace Algo
 	* @return the result of mapping and then summing all the elements of Input
 	*/
 	template <typename T, typename A, typename MapT>
-	FORCEINLINE T TransformAccumulate(const A& Input, MapT MapOp, T init)
+	FORCEINLINE T TransformAccumulate(const A& Input, MapT MapOp, T Init)
 	{
-		return TransformAccumulate(Input, MapOp, init, TPlus<>());
+		return TransformAccumulate(Input, MapOp, MoveTemp(Init), TPlus<>());
 	}
 }
diff --git a/Engine/Source/Runtime/Core/Public/Containers/Algo/Copy.h b/Engine/Source/Runtime/Core/Public/Containers/Algo/Copy.h
new file mode 100644
--- /dev/null
+++ b/Engine/Source/Runtime/Core/Public/Containers/Algo/Copy.h
@@ -0,0 +1,40 @@
+// Copyright 1998-2016 Epic Games, Inc. All Rights Reserved.
+
+#pragma once
+
+namespace Algo
+{
+	/**
+	 * Conditionally copies a range into a container
+	 *
+	 * @param  Input  Any iterable type
+	 * @param  Output  Container to hold the output
+	 * @param  Predicate  Condition which returns true for elements that should be copied and false for elements that should be skipped
+	 */
+	template <typename InT, typename OutT, typename PredicateT>
+	FORCEINLINE void CopyIf(const InT& Input, OutT& Output, PredicateT Predicate)
+	{
+		for (auto&& Value : Input)
+		{
+			if (Predicate(Value))
+			{
+				Output.Add(Value);
+			}
+		}
+	}
+
+	/**
+	 * Copies a range into a container
+	 *
+	 * @param  Input  Any iterable type
+	 * @param  Output  Container to hold the output
+	 */
+	template <typename InT, typename OutT>
+	FORCEINLINE void Copy(const InT& Input, OutT& Output)
+	{
+		for (auto&& Value : Input)
+		{
+			Output.Add(Value);
+		}
+	}
+}
diff --git a/Engine/Source/Runtime/Core/Public/Containers/Algo/Transform.h b/Engine/Source/Runtime/Core/Public/Containers/Algo/Transform.h
new file mode 100644
--- /dev/null
+++ b/Engine/Source/Runtime/Core/Public/Containers/Algo/Transform.h
@@ -0,0 +1,42 @@
+// Copyright 1998-2016 Epic Games, Inc. All Rights Reserved.
+
+#pragma once
+
+namespace Algo
+{
+	/**
+	 * Conditionally applies a transform to a range and stores the results into a container
+	 *
+	 * @param  Input  Any iterable type
+	 * @param  Output  Container to hold the output
+	 * @param  Predicate  Condition which returns true for elements that should be transformed and false for elements that should be skipped
+	 * @param  Transform  Transformation operation
+	 */
+	template <typename InT, typename OutT, typename PredicateT, typename TransformT>
+	FORCEINLINE void TransformIf(const InT& Input, OutT& Output, PredicateT Predicate, TransformT Transform)
+	{
+		for (auto&& Value : Input)
+		{
+			if (Predicate(Value))
+			{
+				Output.Add(Transform(Value));
+			}
+		}
+	}
+
+	/**
+	 * Applies a transform to a range and stores the results into a container
+	 *
+	 * @param  Input  Any iterable type
+	 * @param  Output  Container to hold the output
+	 * @param  Transform  Transformation operation
+	 */
+	template <typename InT, typename OutT, typename TransformT>
+	FORCEINLINE void Transform(const InT& Input, OutT& Output, TransformT Transform)
+	{
+		for (auto&& Value : Input)
+		{
+			Output.Add(Transform(Value));
+		}
+	}
+}
