Fixed: Crash when performing an undo after generating LOD of a level

diff --git a/Engine/Source/Editor/WorldBrowser/Private/Tiles/WorldTileModel.cpp b/Engine/Source/Editor/WorldBrowser/Private/Tiles/WorldTileModel.cpp
--- a/Engine/Source/Editor/WorldBrowser/Private/Tiles/WorldTileModel.cpp
+++ b/Engine/Source/Editor/WorldBrowser/Private/Tiles/WorldTileModel.cpp
@@ -35,6 +35,7 @@ FWorldTileModel::FWorldTileModel(FWorldTileCollectionModel& InWorldModel, int32
 	TileDetails->AddToRoot();
 
 	// Subscribe to tile properties changes
+	// Un-subscribe in dtor if new is added!
 	TileDetails->PostUndoEvent.AddRaw(this, &FWorldTileModel::OnPostUndoEvent);
 	TileDetails->PositionChangedEvent.AddRaw(this, &FWorldTileModel::OnPositionPropertyChanged);
 	TileDetails->ParentPackageNameChangedEvent.AddRaw(this, &FWorldTileModel::OnParentPackageNamePropertyChanged);
@@ -88,9 +89,13 @@ FWorldTileModel::~FWorldTileModel()
 {
 	if (TileDetails)
 	{
+		TileDetails->PostUndoEvent.RemoveAll(this);
 		TileDetails->PositionChangedEvent.RemoveAll(this);
 		TileDetails->ParentPackageNameChangedEvent.RemoveAll(this);
-		
+		TileDetails->LODSettingsChangedEvent.RemoveAll(this);
+		TileDetails->ZOrderChangedEvent.RemoveAll(this);
+		TileDetails->HideInTileViewChangedEvent.RemoveAll(this);
+				
 		TileDetails->RemoveFromRoot();
 		TileDetails->MarkPendingKill();
 	}
