GitHub 4851 : Fix layer info property bNoWeightBlend not being saved after tiled landscape import (contributed by [REDACTED_PERSON])
#4851

diff --git a/Engine/Source/Editor/WorldBrowser/Private/Tiles/WorldTileCollectionModel.cpp b/Engine/Source/Editor/WorldBrowser/Private/Tiles/WorldTileCollectionModel.cpp
--- a/Engine/Source/Editor/WorldBrowser/Private/Tiles/WorldTileCollectionModel.cpp
+++ b/Engine/Source/Editor/WorldBrowser/Private/Tiles/WorldTileCollectionModel.cpp
@@ -1574,6 +1574,7 @@ static void SetupLandscapeImportLayers(const FTiledLandscapeImportSettings& InIm
 		}
 
 		LayerImportInfo.LayerInfo = GetLandscapeLayerInfoObject(LayerImportInfo.LayerName, ContentPath);
+		LayerImportInfo.LayerInfo->Modify();
 		LayerImportInfo.LayerInfo->bNoWeightBlend = LayerSettings.bNoBlendWeight;
 
 		OutLayerInfo.Add(LayerImportInfo);
