Fix for crash when level list contains unloaded levels

diff --git a/Engine/Source/Editor/WorldBrowser/Private/WorldBrowserModule.cpp b/Engine/Source/Editor/WorldBrowser/Private/WorldBrowserModule.cpp
--- a/Engine/Source/Editor/WorldBrowser/Private/WorldBrowserModule.cpp
+++ b/Engine/Source/Editor/WorldBrowser/Private/WorldBrowserModule.cpp
@@ -44,17 +44,20 @@ void FWorldBrowserModule::BuildLevelMenu(FMenuBuilder& MenuBuilder)
 		FLevelModelList ModelList = Model->GetFilteredLevels();
 		for (TSharedPtr<FLevelModel> LevelModel : ModelList)
 		{
-			FUIAction Action(FExecuteAction::CreateRaw(this, &FWorldBrowserModule::SetCurrentSublevel, LevelModel->GetLevelObject()),
-				FCanExecuteAction(),
-				FIsActionChecked::CreateRaw(this, &FWorldBrowserModule::IsCurrentSublevel, LevelModel->GetLevelObject()));
-			MenuBuilder.AddMenuEntry(FText::FromString(LevelModel->GetDisplayName()), FText::GetEmpty(), FSlateIcon(), Action, NAME_None, EUserInterfaceActionType::Button);
+			if (LevelModel->GetLevelObject())
+			{
+				FUIAction Action(FExecuteAction::CreateRaw(this, &FWorldBrowserModule::SetCurrentSublevel, LevelModel->GetLevelObject()),
+					FCanExecuteAction(),
+					FIsActionChecked::CreateRaw(this, &FWorldBrowserModule::IsCurrentSublevel, LevelModel->GetLevelObject()));
+				MenuBuilder.AddMenuEntry(FText::FromString(LevelModel->GetDisplayName()), FText::GetEmpty(), FSlateIcon(), Action, NAME_None, EUserInterfaceActionType::Button);
+			}
 		}
 	}
 }
 
 bool FWorldBrowserModule::IsCurrentSublevel(ULevel* InLevel)
 {
-	return InLevel->IsCurrentLevel();
+	return InLevel ? InLevel->IsCurrentLevel() : false;
 }
 
 void FWorldBrowserModule::SetCurrentSublevel(ULevel* InLevel)
