Fixed initialization of FWorldTileModel in world composition when Level is loaded in the case where LevelBounds have its bAutoUpdateBounds set to false.

diff --git a/Engine/Source/Editor/WorldBrowser/Private/Tiles/WorldTileModel.cpp b/Engine/Source/Editor/WorldBrowser/Private/Tiles/WorldTileModel.cpp
--- a/Engine/Source/Editor/WorldBrowser/Private/Tiles/WorldTileModel.cpp
+++ b/Engine/Source/Editor/WorldBrowser/Private/Tiles/WorldTileModel.cpp
@@ -437,11 +437,6 @@ FVector FWorldTileModel::GetLevelCurrentPosition() const
 	return FVector::ZeroVector;
 }
 
-FVector2D FWorldTileModel::GetLandscapeComponentSize() const
-{
-	return LandscapeComponentSize;
-}
-
 void FWorldTileModel::SetLevelPosition(const FIntVector& InPosition)
 {
 	// Parent absolute position
@@ -522,9 +517,6 @@ void FWorldTileModel::Update()
 	if (!IsRootTile())
 	{
 		Landscape = NULL;
-		LandscapeComponentsXY.Empty();
-		LandscapeComponentSize = FVector2D(0.f, 0.f);
-		LandscapeComponentsRectXY = FIntRect(FIntPoint(MAX_int32, MAX_int32), FIntPoint(MIN_int32, MIN_int32));
 		
 		ULevel* Level = GetLevelObject();
 		// Receive tile info from world composition
@@ -691,6 +683,10 @@ void FWorldTileModel::OnLevelAddedToWorld(ULevel* InLevel)
 	FLevelModel::OnLevelAddedToWorld(InLevel);
 
 	EnsureLevelHasBoundsActor();
+
+	// Manually call Update to make sure WorldTileModel is properly initialized (don't rely on Level Bounds changes as it will not be called if bAutoUpdateBounds is set to false).
+	Update();
+
 	LoadedLevel.Get()->LevelBoundsActorUpdated().AddRaw(this, &FWorldTileModel::OnLevelBoundsActorUpdated);
 }
 
diff --git a/Engine/Source/Editor/WorldBrowser/Private/Tiles/WorldTileModel.h b/Engine/Source/Editor/WorldBrowser/Private/Tiles/WorldTileModel.h
--- a/Engine/Source/Editor/WorldBrowser/Private/Tiles/WorldTileModel.h
+++ b/Engine/Source/Editor/WorldBrowser/Private/Tiles/WorldTileModel.h
@@ -147,9 +147,6 @@ public:
 	/** @return ULevel bounding box in shifted space*/
 	FBox GetLevelBounds() const override;
 	
-	/** @return Landscape component world size */
-	FVector2D GetLandscapeComponentSize() const;
-
 	/** Translate level center to new position */
 	void SetLevelPosition(const FIntVector& InPosition);
 
@@ -227,9 +224,4 @@ public:
 private:
 	/** Whether level has landscape components in it */
 	TWeakObjectPtr<ALandscapeProxy>			Landscape;
-	FVector									LandscapeComponentResolution;
-	TArray<FIntPoint>						LandscapeComponentsXY;
-	FIntRect								LandscapeComponentsRectXY;
-	FVector2D								LandscapeComponentSize;
-	FVector2D								LandscapePosition;	
 };
