Fix Create Adjacent Landscape Proxy not setting the proper LandscapeSectionOffset

diff --git a/Engine/Source/Editor/WorldBrowser/Private/Tiles/WorldTileCollectionModel.cpp b/Engine/Source/Editor/WorldBrowser/Private/Tiles/WorldTileCollectionModel.cpp
--- a/Engine/Source/Editor/WorldBrowser/Private/Tiles/WorldTileCollectionModel.cpp
+++ b/Engine/Source/Editor/WorldBrowser/Private/Tiles/WorldTileCollectionModel.cpp
@@ -1471,10 +1471,9 @@ void FWorldTileCollectionModel::AddLandscapeProxy_Executed(FWorldTileModel::EWor
 		Levels.Add(NewLevelModel);
 
 		ALandscapeProxy* SourceLandscape = LandscapeTileModel->GetLandscape();
-		FIntVector SourceTileOffset = LandscapeTileModel->GetAbsoluteLevelPosition();
-
+	
 		NewLevelModel->SetVisible(false);
-		NewLevelModel->CreateAdjacentLandscapeProxy(SourceLandscape, SourceTileOffset, InWhere);
+		NewLevelModel->CreateAdjacentLandscapeProxy(SourceLandscape, InWhere);
 		ShowLevels(Levels);
 	}
 }
diff --git a/Engine/Source/Editor/WorldBrowser/Private/Tiles/WorldTileModel.cpp b/Engine/Source/Editor/WorldBrowser/Private/Tiles/WorldTileModel.cpp
--- a/Engine/Source/Editor/WorldBrowser/Private/Tiles/WorldTileModel.cpp
+++ b/Engine/Source/Editor/WorldBrowser/Private/Tiles/WorldTileModel.cpp
@@ -437,7 +437,7 @@ FVector FWorldTileModel::GetLevelCurrentPosition() const
 	return FVector::ZeroVector;
 }
 
-void FWorldTileModel::SetLevelPosition(const FIntVector& InPosition)
+void FWorldTileModel::SetLevelPosition(const FIntVector& InPosition, const FIntPoint* InLandscapeSectionOffset)
 {
 	// Parent absolute position
 	TSharedPtr<FWorldTileModel> ParentModel = StaticCastSharedPtr<FWorldTileModel>(GetParent());
@@ -445,6 +445,7 @@ void FWorldTileModel::SetLevelPosition(const FIntVector& InPosition)
 	
 	// Actual offset
 	FIntVector Offset = InPosition - TileDetails->AbsolutePosition;
+	FIntPoint LandscapeOffset = InLandscapeSectionOffset ? *InLandscapeSectionOffset : FIntPoint(Offset.X, Offset.Y);
 
 	TileDetails->Modify();
 	
@@ -484,7 +485,7 @@ void FWorldTileModel::SetLevelPosition(const FIntVector& InPosition)
 
 	if (IsLandscapeBased())
 	{
-		UpdateLandscapeSectionsOffset(FIntPoint(Offset.X, Offset.Y)); // section offset is 2D 
+		UpdateLandscapeSectionsOffset(LandscapeOffset); // section offset is 2D 
 		bool bShowWarnings = true;
 		ULandscapeInfo::RecreateLandscapeInfo(LevelCollectionModel.GetWorld(), bShowWarnings);
 	}
@@ -894,7 +895,7 @@ void FWorldTileModel::OnHideInTileViewChanged()
 	OnLevelInfoUpdated();
 }
 
-bool FWorldTileModel::CreateAdjacentLandscapeProxy(ALandscapeProxy* SourceLandscape, const FIntVector& SourceTileOffset, FWorldTileModel::EWorldDirections InWhere)
+bool FWorldTileModel::CreateAdjacentLandscapeProxy(ALandscapeProxy* SourceLandscape, FWorldTileModel::EWorldDirections InWhere)
 {
 	if (!IsLoaded())	
 	{
@@ -906,6 +907,8 @@ bool FWorldTileModel::CreateAdjacentLandscapeProxy(ALandscapeProxy* SourceLandsc
 	FVector SourceLandscapeScale = SourceLandscape->GetRootComponent()->GetComponentToWorld().GetScale3D();
 	FIntRect SourceLandscapeRect = SourceLandscape->GetBoundingRect();
 	FIntPoint SourceLandscapeSize = SourceLandscapeRect.Size();
+	FIntPoint LandscapeSectionOffset = FIntPoint((SourceLandscape->LandscapeSectionOffset.X+SourceLandscapeRect.Min.X)*SourceLandscapeScale.X, 
+												 (SourceLandscape->LandscapeSectionOffset.Y+SourceLandscapeRect.Min.Y)*SourceLandscapeScale.Y);
 
 	FLandscapeImportSettings ImportSettings = {};
 	ImportSettings.LandscapeGuid		= SourceLandscape->GetLandscapeGuid();
@@ -945,15 +948,19 @@ bool FWorldTileModel::CreateAdjacentLandscapeProxy(ALandscapeProxy* SourceLandsc
 		switch (InWhere)
 		{
 		case FWorldTileModel::XNegative:
+			LandscapeSectionOffset += FIntPoint(-SourceLandscapeScale.X*SourceLandscapeSize.X, 0);
 			ProxyOffset += FVector(-SourceLandscapeScale.X*SourceLandscapeSize.X, 0.f, 0.f);
 			break;
 		case FWorldTileModel::XPositive:
+			LandscapeSectionOffset += FIntPoint(+SourceLandscapeScale.X*SourceLandscapeSize.X, 0);
 			ProxyOffset += FVector(+SourceLandscapeScale.X*SourceLandscapeSize.X, 0.f, 0.f);
 			break;
 		case FWorldTileModel::YNegative:
+			LandscapeSectionOffset += FIntPoint(0, -SourceLandscapeScale.Y*SourceLandscapeSize.Y);
 			ProxyOffset += FVector(0.f, -SourceLandscapeScale.Y*SourceLandscapeSize.Y, 0.f);
 			break;
 		case FWorldTileModel::YPositive:
+			LandscapeSectionOffset += FIntPoint(0, +SourceLandscapeScale.Y*SourceLandscapeSize.Y);
 			ProxyOffset += FVector(0.f, +SourceLandscapeScale.Y*SourceLandscapeSize.Y, 0.f);
 			break;
 		}
@@ -962,7 +969,7 @@ bool FWorldTileModel::CreateAdjacentLandscapeProxy(ALandscapeProxy* SourceLandsc
 		FIntVector IntOffset = FIntVector(ProxyOffset) + LevelCollectionModel.GetWorld()->OriginLocation;
 
 		// Move level with landscape proxy to desired position
-		SetLevelPosition(IntOffset);
+		SetLevelPosition(IntOffset, &LandscapeSectionOffset);
 		return true;
 	}
 
diff --git a/Engine/Source/Editor/WorldBrowser/Private/Tiles/WorldTileModel.h b/Engine/Source/Editor/WorldBrowser/Private/Tiles/WorldTileModel.h
--- a/Engine/Source/Editor/WorldBrowser/Private/Tiles/WorldTileModel.h
+++ b/Engine/Source/Editor/WorldBrowser/Private/Tiles/WorldTileModel.h
@@ -148,7 +148,7 @@ public:
 	FBox GetLevelBounds() const override;
 	
 	/** Translate level center to new position */
-	void SetLevelPosition(const FIntVector& InPosition);
+	void SetLevelPosition(const FIntVector& InPosition, const FIntPoint* InLandscapeSectionOffset = nullptr);
 
 	/** Recursively sort all children by name */
 	void SortRecursive();
@@ -163,7 +163,7 @@ public:
 	ULevelStreaming* GetAssosiatedStreamingLevel() { return GetAssociatedStreamingLevel(); }
 
 	/**  */
-	bool CreateAdjacentLandscapeProxy(ALandscapeProxy* SourceLandscape, const FIntVector& SourceTileOffset, FWorldTileModel::EWorldDirections InWhere);
+	bool CreateAdjacentLandscapeProxy(ALandscapeProxy* SourceLandscape, FWorldTileModel::EWorldDirections InWhere);
 
 	/**  */
 	ALandscapeProxy* ImportLandscapeTile(const FLandscapeImportSettings& Settings);
