More clarification on what is Distance in world composition tile LOD

diff --git a/Engine/Source/Editor/WorldBrowser/Private/Tiles/WorldTileDetails.cpp b/Engine/Source/Editor/WorldBrowser/Private/Tiles/WorldTileDetails.cpp
--- a/Engine/Source/Editor/WorldBrowser/Private/Tiles/WorldTileDetails.cpp
+++ b/Engine/Source/Editor/WorldBrowser/Private/Tiles/WorldTileDetails.cpp
@@ -11,7 +11,7 @@
 FTileLODEntryDetails::FTileLODEntryDetails()
 	// Initialize properties with default values from FWorldTileLODInfo
 	: LODIndex(0)
-	, Distance(FWorldTileLODInfo().RelativeStreamingDistance)
+	, RelativeDistance(FWorldTileLODInfo().RelativeStreamingDistance)
 {
 }
 
@@ -44,12 +44,12 @@ void UWorldTileDetails::SetInfo(const FWorldTileInfo& Info, ULevel* Level)
 	{
 		if (Info.LODList.IsValidIndex(i))
 		{
-			LODEntries[i]->Distance = Info.LODList[i].RelativeStreamingDistance;
+			LODEntries[i]->RelativeDistance = Info.LODList[i].RelativeStreamingDistance;
 			LODEntries[i]->SimplificationDetails = Level ? Level->LevelSimplification[i] : FLevelSimplificationDetails();
 		}
 		else
 		{
-			LODEntries[i]->Distance = FWorldTileLODInfo().RelativeStreamingDistance;
+			LODEntries[i]->RelativeDistance = FWorldTileLODInfo().RelativeStreamingDistance;
 			LODEntries[i]->SimplificationDetails = FLevelSimplificationDetails();
 		}
 	}
@@ -72,7 +72,7 @@ FWorldTileInfo UWorldTileDetails::GetInfo() const
 	const FTileLODEntryDetails* LODEntries[WORLDTILE_LOD_MAX_INDEX] = {&LOD1, &LOD2, &LOD3, &LOD4};
 	for (int32 i = 0; i < Info.LODList.Num(); ++i)
 	{
-		Info.LODList[i].RelativeStreamingDistance = LODEntries[i]->Distance;
+		Info.LODList[i].RelativeStreamingDistance = LODEntries[i]->RelativeDistance;
 	}
 	
 	return Info;
diff --git a/Engine/Source/Editor/WorldBrowser/Private/Tiles/WorldTileDetails.h b/Engine/Source/Editor/WorldBrowser/Private/Tiles/WorldTileDetails.h
--- a/Engine/Source/Editor/WorldBrowser/Private/Tiles/WorldTileDetails.h
+++ b/Engine/Source/Editor/WorldBrowser/Private/Tiles/WorldTileDetails.h
@@ -22,9 +22,9 @@ struct FTileLODEntryDetails
 	UPROPERTY(Category=LODDetails, VisibleAnywhere)
 	int32		LODIndex;
 			
-	// Relative to original tile streaming distance
+	// Relative to original tile streaming distance (Absolute Distance = Layer Streaming Distance + Relative Distance)
 	UPROPERTY(Category=LODDetails, EditAnywhere, meta=(ClampMin = "10", ClampMax = "10000000", UIMin = "10", UIMax = "500000"))
-	int32		Distance;
+	int32		RelativeDistance;
 
 	UPROPERTY(Category=ReductionSettings, EditAnywhere)	
 	FLevelSimplificationDetails SimplificationDetails;
diff --git a/Engine/Source/Editor/WorldBrowser/Private/Tiles/WorldTileDetailsCustomization.cpp b/Engine/Source/Editor/WorldBrowser/Private/Tiles/WorldTileDetailsCustomization.cpp
--- a/Engine/Source/Editor/WorldBrowser/Private/Tiles/WorldTileDetailsCustomization.cpp
+++ b/Engine/Source/Editor/WorldBrowser/Private/Tiles/WorldTileDetailsCustomization.cpp
@@ -270,7 +270,7 @@ void FTileLODEntryDetailsCustomization::CustomizeChildren(TSharedRef<IPropertyHa
 		);
 	
 	TSharedPtr<IPropertyHandle> DistanceProperty = StructPropertyHandle->GetChildHandle(
-		GET_MEMBER_NAME_CHECKED(FTileLODEntryDetails, Distance)
+		GET_MEMBER_NAME_CHECKED(FTileLODEntryDetails, RelativeDistance)
 		);
 	
 	TSharedPtr<IPropertyHandle> SimplificationDetails = StructPropertyHandle->GetChildHandle(
