Fixed crash when attempting to load a world composition sublevel after deleting the map.
#rnx

diff --git a/Engine/Source/Editor/UnrealEd/Private/ObjectTools.cpp b/Engine/Source/Editor/UnrealEd/Private/ObjectTools.cpp
--- a/Engine/Source/Editor/UnrealEd/Private/ObjectTools.cpp
+++ b/Engine/Source/Editor/UnrealEd/Private/ObjectTools.cpp
@@ -1944,6 +1944,9 @@ namespace ObjectTools
 				}
 			}
 		}
+
+		// Let the level browser that we deleted a level (must happen after physically deleting the package file as it will rescan the folders)
+		FEditorDelegates::RefreshLevelBrowser.Broadcast();
 	}
 
 	int32 DeleteAssets( const TArray<FAssetData>& AssetsToDelete, bool bShowConfirmation )
diff --git a/Engine/Source/Editor/WorldBrowser/Private/LevelCollectionModel.cpp b/Engine/Source/Editor/WorldBrowser/Private/LevelCollectionModel.cpp
--- a/Engine/Source/Editor/WorldBrowser/Private/LevelCollectionModel.cpp
+++ b/Engine/Source/Editor/WorldBrowser/Private/LevelCollectionModel.cpp
@@ -52,6 +52,7 @@ FLevelCollectionModel::FLevelCollectionModel()
 	, bSelectionHasChanged(true)
 	, bUpdatingLevelsSelection(false)
 {
+	FEditorDelegates::RefreshLevelBrowser.AddRaw(this, &FLevelCollectionModel::PopulateLevelsList);
 }
 
 FLevelCollectionModel::~FLevelCollectionModel()
@@ -68,6 +69,8 @@ FLevelCollectionModel::~FLevelCollectionModel()
 	{
 		CurrentWorld->OnSelectedLevelsChanged().RemoveAll(this);
 	}
+
+	FEditorDelegates::RefreshLevelBrowser.RemoveAll(this);
 }
 
 void FLevelCollectionModel::Initialize(UWorld* InWorld)
diff --git a/Engine/Source/Editor/WorldBrowser/Private/StreamingLevels/StreamingLevelCollectionModel.cpp b/Engine/Source/Editor/WorldBrowser/Private/StreamingLevels/StreamingLevelCollectionModel.cpp
--- a/Engine/Source/Editor/WorldBrowser/Private/StreamingLevels/StreamingLevelCollectionModel.cpp
+++ b/Engine/Source/Editor/WorldBrowser/Private/StreamingLevels/StreamingLevelCollectionModel.cpp
@@ -37,14 +37,11 @@ FStreamingLevelCollectionModel::FStreamingLevelCollectionModel()
 	{
 		AddedLevelStreamingClass = DefaultLevelStreamingClass;
 	}
-
-	FEditorDelegates::RefreshLevelBrowser.AddRaw(this, &FStreamingLevelCollectionModel::PopulateLevelsList);
 }
 
 FStreamingLevelCollectionModel::~FStreamingLevelCollectionModel()
 {
 	GEditor->UnregisterForUndo( this );
-	FEditorDelegates::RefreshLevelBrowser.RemoveAll(this);
 }
 
 void FStreamingLevelCollectionModel::Initialize(UWorld* InWorld)
