Can now remove a unloaded/failed to load streaming level from the levels tab

diff --git a/Engine/Source/Editor/WorldBrowser/Private/LevelCollectionModel.cpp b/Engine/Source/Editor/WorldBrowser/Private/LevelCollectionModel.cpp
--- a/Engine/Source/Editor/WorldBrowser/Private/LevelCollectionModel.cpp
+++ b/Engine/Source/Editor/WorldBrowser/Private/LevelCollectionModel.cpp
@@ -664,6 +664,11 @@ void FLevelCollectionModel::UnloadLevels(const FLevelModelList& InLevelList)
 				EditorLevelUtils::RemoveLevelFromWorld(Level);
 			}
 		}
+		else if (ULevelStreaming* StreamingLevel = Cast<ULevelStreaming>(LevelModel->GetNodeObject()))
+		{
+			StreamingLevel->MarkPendingKill();
+			ThisWorld->RemoveStreamingLevel(StreamingLevel);
+		}
 	}
 
 	BroadcastPostLevelsUnloaded();
diff --git a/Engine/Source/Editor/WorldBrowser/Private/StreamingLevels/StreamingLevelCollectionModel.cpp b/Engine/Source/Editor/WorldBrowser/Private/StreamingLevels/StreamingLevelCollectionModel.cpp
--- a/Engine/Source/Editor/WorldBrowser/Private/StreamingLevels/StreamingLevelCollectionModel.cpp
+++ b/Engine/Source/Editor/WorldBrowser/Private/StreamingLevels/StreamingLevelCollectionModel.cpp
@@ -184,7 +184,7 @@ void FStreamingLevelCollectionModel::BindCommands()
 	
 	ActionList.MapAction( Commands.World_RemoveSelectedLevels,
 		FExecuteAction::CreateSP( this, &FStreamingLevelCollectionModel::UnloadSelectedLevels_Executed  ),
-		FCanExecuteAction::CreateSP( this, &FLevelCollectionModel::AreAllSelectedLevelsEditable ) );
+		FCanExecuteAction::CreateSP( this, &FStreamingLevelCollectionModel::AreAllSelectedLevelsRemovable ) );
 	
 	ActionList.MapAction( Commands.World_MergeSelectedLevels,
 		FExecuteAction::CreateSP( this, &FStreamingLevelCollectionModel::MergeSelectedLevels_Executed  ),
@@ -566,6 +566,19 @@ bool FStreamingLevelCollectionModel::IsStreamingMethodChecked(UClass* InClass) c
 	return false;
 }
 
+bool FStreamingLevelCollectionModel::AreAllSelectedLevelsRemovable() const
+{
+	for (const TSharedPtr<FLevelModel>& LevelModel : SelectedLevelsList)
+	{
+		if (LevelModel->IsLocked() || LevelModel->IsPersistent())
+		{
+			return false;
+		}
+	}
+
+	return AreAnyLevelsSelected();
+}
+
 void FStreamingLevelCollectionModel::SetStreamingLevelsClass_Executed(UClass* InClass)
 {
 	// First prompt to save the selected levels, as changing the streaming method will unload/reload them
diff --git a/Engine/Source/Editor/WorldBrowser/Private/StreamingLevels/StreamingLevelCollectionModel.h b/Engine/Source/Editor/WorldBrowser/Private/StreamingLevels/StreamingLevelCollectionModel.h
--- a/Engine/Source/Editor/WorldBrowser/Private/StreamingLevels/StreamingLevelCollectionModel.h
+++ b/Engine/Source/Editor/WorldBrowser/Private/StreamingLevels/StreamingLevelCollectionModel.h
@@ -106,6 +106,9 @@ private:
 	/** Checks if the passed in streaming method is the current streaming method */
 	bool IsStreamingMethodChecked(UClass* InClass) const;
 
+	/** @return	whether the selected levels are allowed to be removed from the world */
+	bool AreAllSelectedLevelsRemovable() const;
+
 	/** Changes the streaming method for the selected levels. */
 	void SetStreamingLevelsClass_Executed(UClass* InClass);
 
