PR #5060: Initialize booleans in constructor for WorldTreeItems (Contributed by [REDACTED_PERSON])

diff --git a/Engine/Source/Editor/WorldBrowser/Public/IWorldTreeItem.h b/Engine/Source/Editor/WorldBrowser/Public/IWorldTreeItem.h
--- a/Engine/Source/Editor/WorldBrowser/Public/IWorldTreeItem.h
+++ b/Engine/Source/Editor/WorldBrowser/Public/IWorldTreeItem.h
@@ -56,6 +56,8 @@ namespace WorldHierarchy
 				: bExpanded(1)
 				, bFilteredOut(0)
 				, bChildrenRequiresSort(1)
+				, bVisible(1)
+				, bLocked(0)
 			{}
 		};
 
