Copy and paste for level transforms
#rnx

diff --git a/Engine/Source/Editor/WorldBrowser/Private/StreamingLevels/StreamingLevelCustomization.cpp b/Engine/Source/Editor/WorldBrowser/Private/StreamingLevels/StreamingLevelCustomization.cpp
--- a/Engine/Source/Editor/WorldBrowser/Private/StreamingLevels/StreamingLevelCustomization.cpp
+++ b/Engine/Source/Editor/WorldBrowser/Private/StreamingLevels/StreamingLevelCustomization.cpp
@@ -19,6 +19,9 @@
 #include "StreamingLevels/StreamingLevelModel.h"
 #include "StreamingLevels/StreamingLevelEdMode.h"
 #include "Widgets/Input/SNumericEntryBox.h"
+#include "Math/UnitConversion.h"
+#include "HAL/PlatformApplicationMisc.h"
+#include "ScopedTransaction.h"
 
 #define LOCTEXT_NAMESPACE "WorldBrowser"
 
@@ -48,6 +51,8 @@ void FStreamingLevelCustomization::CustomizeDetails(IDetailLayoutBuilder& Detail
 
 	// Add Position property
 	LevelStreamingCategory.AddCustomRow(LOCTEXT("Position", "Position"))
+		.CopyAction(CreateCopyAction(ELevelTransformField::Location))
+		.PasteAction(CreatePasteAction(ELevelTransformField::Location))
 		.NameContent()
 		[
 			SNew(STextBlock)
@@ -73,6 +78,8 @@ void FStreamingLevelCustomization::CustomizeDetails(IDetailLayoutBuilder& Detail
 	
 	// Add Yaw Rotation property
 	LevelStreamingCategory.AddCustomRow(LOCTEXT("Rotation", "Rotation"))
+		.CopyAction(CreateCopyAction(ELevelTransformField::Rotation))
+		.PasteAction(CreatePasteAction(ELevelTransformField::Rotation))
 		.NameContent()
 		[
 			SNew(STextBlock)
@@ -93,16 +100,16 @@ void FStreamingLevelCustomization::CustomizeDetails(IDetailLayoutBuilder& Detail
 				.IsEnabled(this, &FStreamingLevelCustomization::LevelEditTextTransformAllowed)
 				.Font(FEditorStyle::GetFontStyle(TEXT("PropertyWindow.NormalFont")))
 				.UndeterminedString(LOCTEXT("MultipleValues", "Multiple Values"))
-				//.Delta(90)
 				.AllowSpin(true)
 				.MinValue(0)
 				.MaxValue(360 - 1)
 				.MinSliderValue(0)
 				.MaxSliderValue(360 - 1)
 				.Value(this, &FStreamingLevelCustomization::GetLevelRotation)
-				.OnValueChanged(this, &FStreamingLevelCustomization::OnSetLevelRotation)
-				.OnBeginSliderMovement(this, &FStreamingLevelCustomization::OnBeginLevelRotatonSlider)
-				.OnEndSliderMovement(this, &FStreamingLevelCustomization::OnEndLevelRotatonSlider)
+				.OnValueChanged(this, &FStreamingLevelCustomization::OnChangedLevelRotation)
+				.OnValueCommitted(this, &FStreamingLevelCustomization::OnSetLevelRotation)
+				.OnBeginSliderMovement(this, &FStreamingLevelCustomization::OnBeginLevelRotationSlider)
+				.OnEndSliderMovement(this, &FStreamingLevelCustomization::OnEndLevelRotationSlider)
 				.LabelPadding(0)					
 				.Label()
 				[
@@ -212,8 +219,84 @@ bool FStreamingLevelCustomization::OnShouldSetEditorStreamingVolume(const FAsset
 	return false;
 }
 
+void FStreamingLevelCustomization::OnCopy(ELevelTransformField::Type TransformField)
+{
+
+	FString CopyStr;
+	switch (TransformField)
+	{
+	case ELevelTransformField::Location:
+		CopyStr = FString::Printf(TEXT("(X=%f,Y=%f,Z=%f)"), OnGetLevelPosition(0).GetValue(), OnGetLevelPosition(1).GetValue(), OnGetLevelPosition(2).GetValue());
+		break;
+	case ELevelTransformField::Rotation:
+		CopyStr = FString::Printf(TEXT("(Pitch=%f,Yaw=%f,Roll=%f)"), 0.0f, (float)GetLevelRotation().GetValue(), 0.0f);
+		break;
+	default:
+		break;
+	}
+
+	if (!CopyStr.IsEmpty())
+	{
+		FPlatformApplicationMisc::ClipboardCopy(*CopyStr);
+	}
+}
+
+void FStreamingLevelCustomization::OnPaste(ELevelTransformField::Type TransformField)
+{
+	FString PastedText;
+	FPlatformApplicationMisc::ClipboardPaste(PastedText);
+
+	switch (TransformField)
+	{
+	case ELevelTransformField::Location:
+	{
+		FVector Location;
+		if (Location.InitFromString(PastedText))
+		{
+			FScopedTransaction Transaction(LOCTEXT("PasteLocation", "Paste Location"));
+			OnSetLevelPosition(Location);
+		}
+	}
+	break;
+	case ELevelTransformField::Rotation:
+	{
+		FRotator Rotation;
+		PastedText.ReplaceInline(TEXT("Pitch="), TEXT("P="));
+		PastedText.ReplaceInline(TEXT("Yaw="), TEXT("Y="));
+		PastedText.ReplaceInline(TEXT("Roll="), TEXT("R="));
+		if (Rotation.InitFromString(PastedText))
+		{
+			FScopedTransaction Transaction(LOCTEXT("PasteRotation", "Paste Rotation"));
+			OnSetLevelRotation(Rotation.Yaw, ETextCommit::OnEnter);
+		}
+	}
+	break;
+	default:
+		break;
+	}
+}
+
+FUIAction FStreamingLevelCustomization::CreateCopyAction(ELevelTransformField::Type TransformField) const
+{
+	return
+		FUIAction
+		(
+			FExecuteAction::CreateSP(this, &FStreamingLevelCustomization::OnCopy, TransformField)
+		);
+}
+
+FUIAction FStreamingLevelCustomization::CreatePasteAction(ELevelTransformField::Type TransformField) const
+{
+	return
+		FUIAction(FExecuteAction::CreateSP(this, &FStreamingLevelCustomization::OnPaste, TransformField));
+}
+
 void FStreamingLevelCustomization::OnSetLevelPosition( float NewValue, ETextCommit::Type CommitInfo, int32 Axis )
 {
+	if (CommitInfo == ETextCommit::Default)
+	{
+		return;
+	}
 	TSharedPtr<FStreamingLevelCollectionModel> CollectionModel = WorldModel.Pin();
 	if (CollectionModel.IsValid())
 	{
@@ -236,6 +319,28 @@ void FStreamingLevelCustomization::OnSetLevelPosition( float NewValue, ETextComm
 	}
 }
 
+void FStreamingLevelCustomization::OnSetLevelPosition(FVector NewValue)
+{
+	TSharedPtr<FStreamingLevelCollectionModel> CollectionModel = WorldModel.Pin();
+	if (CollectionModel.IsValid())
+	{
+		FLevelModelList SelectedLevels = CollectionModel->GetSelectedLevels();
+		for (auto It = SelectedLevels.CreateIterator(); It; ++It)
+		{
+			TSharedPtr<FStreamingLevelModel> LevelModel = StaticCastSharedPtr<FStreamingLevelModel>(*It);
+			if (LevelModel->IsEditable() && LevelModel->GetLevelStreaming().IsValid())
+			{
+				ULevelStreaming* LevelStreaming = LevelModel->GetLevelStreaming().Get();
+				// Create transform with new translation
+				FTransform LevelTransform = LevelStreaming->LevelTransform;
+				LevelTransform.SetTranslation(NewValue);
+				// Transform level
+				FLevelUtils::SetEditorTransform(LevelStreaming, LevelTransform);
+			}
+		}
+	}
+}
+
 TOptional<float> FStreamingLevelCustomization::OnGetLevelPosition( int32 Axis ) const
 {
 	float AxisVal = 0.f;
@@ -247,19 +352,24 @@ TOptional<float> FStreamingLevelCustomization::OnGetLevelPosition( int32 Axis )
 	return TOptional<float>();
 }
 
-void FStreamingLevelCustomization::OnSetLevelRotation( int32 NewValue )
+void FStreamingLevelCustomization::OnChangedLevelRotation(int32 NewValue)
 {
 	CachedYawValue = NewValue;
-	if (bSliderMovement)
+}
+
+void FStreamingLevelCustomization::OnSetLevelRotation( int32 NewValue, ETextCommit::Type CommitInfo)
+{
+	CachedYawValue = NewValue;
+	if (CommitInfo == ETextCommit::Default)
 	{
 		return;
 	}
 	
-	// Apply rotation when user stops slider drag
+	// Apply rotation when value is committed
 	TSharedPtr<FStreamingLevelCollectionModel> CollectionModel = WorldModel.Pin();
 	if (CollectionModel.IsValid())
 	{
-		FQuat NewRotaion = FRotator(0.f, (float)CachedYawValue.GetValue(), 0.f).Quaternion();
+		FQuat NewRotation = FRotator(0.f, (float)CachedYawValue.GetValue(), 0.f).Quaternion();
 		
 		FLevelModelList SelectedLevels = CollectionModel->GetSelectedLevels();
 		for (auto It = SelectedLevels.CreateIterator(); It; ++It)
@@ -269,9 +379,9 @@ void FStreamingLevelCustomization::OnSetLevelRotation( int32 NewValue )
 			{
 				ULevelStreaming* LevelStreaming = LevelModel->GetLevelStreaming().Get();
 				FTransform LevelTransform = LevelStreaming->LevelTransform;
-				if (LevelTransform.GetRotation() != NewRotaion)
+				if (LevelTransform.GetRotation() != NewRotation)
 				{
-					LevelTransform.SetRotation(NewRotaion);
+					LevelTransform.SetRotation(NewRotation);
 					FLevelUtils::SetEditorTransform(LevelStreaming, LevelTransform);
 				}
 			}
@@ -279,16 +389,15 @@ void FStreamingLevelCustomization::OnSetLevelRotation( int32 NewValue )
 	}
 }
 
-void FStreamingLevelCustomization::OnBeginLevelRotatonSlider()
+void FStreamingLevelCustomization::OnBeginLevelRotationSlider()
 {		
 	CachedYawValue = GetLevelRotation();
 	bSliderMovement = true;
 }
 
-void FStreamingLevelCustomization::OnEndLevelRotatonSlider( int32 NewValue )
+void FStreamingLevelCustomization::OnEndLevelRotationSlider( int32 NewValue )
 {
 	bSliderMovement = false;
-	OnSetLevelRotation(NewValue);
 }
 
 TOptional<int32> FStreamingLevelCustomization::GetLevelRotation() const
diff --git a/Engine/Source/Editor/WorldBrowser/Private/StreamingLevels/StreamingLevelCustomization.h b/Engine/Source/Editor/WorldBrowser/Private/StreamingLevels/StreamingLevelCustomization.h
--- a/Engine/Source/Editor/WorldBrowser/Private/StreamingLevels/StreamingLevelCustomization.h
+++ b/Engine/Source/Editor/WorldBrowser/Private/StreamingLevels/StreamingLevelCustomization.h
@@ -10,6 +10,15 @@ class IDetailChildrenBuilder;
 class IDetailLayoutBuilder;
 class IPropertyHandle;
 
+namespace ELevelTransformField
+{
+	enum Type
+	{
+		Location,
+		Rotation
+	};
+}
+
 /////////////////////////////////////////////////////
 // FStreamingLevelCustomization
 class FStreamingLevelCustomization : public IDetailCustomization
@@ -34,6 +43,13 @@ private:
 	 */
 	void OnSetLevelPosition( float NewValue, ETextCommit::Type CommitInfo, int32 Axis );
 
+	/**
+	 *	Called by the Editor Level Transform column, to set the new values.
+	 *
+	 *	@param NewValue		The new value entered.
+	 */
+	void OnSetLevelPosition(FVector NewValue);
+
 	/**
 	 *	Called by the Editor Level Transform column, to get the current values.
 	 *
@@ -43,24 +59,32 @@ private:
 	 */
 	TOptional<float> OnGetLevelPosition( int32 Axis ) const;
 
+	/**
+	 *	Called by the Editor Level Rotation column, to update the displayed value as we spin
+	 *
+	 *	@param NewValue		The new value entered.
+	 */
+	void OnChangedLevelRotation(int32 NewValue);
+
 	/**
 	 *	Called by the Editor Level Rotation column, to set the new values.
 	 *
 	 *	@param NewValue		The new value entered.
+	 *	@param CommitInfo	How the value was set.
 	 */
-	void OnSetLevelRotation( int32 NewValue );
+	void OnSetLevelRotation( int32 NewValue, ETextCommit::Type CommitInfo);
 
 	/**
 	 *	Called by the Editor Level Rotation column, to set weve started spinning the value.
 	 */
-	void OnBeginLevelRotatonSlider();
+	void OnBeginLevelRotationSlider();
 
 	/**
 	 *	Called by the Editor Level Rotation column, to set weve stopped spinning the value.
 	 *
 	 *	@param NewValue		The new value entered.
 	 */
-	void OnEndLevelRotatonSlider( int32 NewValue );
+	void OnEndLevelRotationSlider( int32 NewValue );
 
 	/**
 	 *	Called by the Editor Level Transform column, to get the current values.
@@ -84,6 +108,26 @@ private:
 	/** Called to determine whether the selected streaming volume can be set or not */
 	bool OnShouldSetEditorStreamingVolume(const FAssetData& AssetData, TSharedRef<IPropertyHandle> ElementProperty) const;
 
+	/**
+	 * Copies the specified transform field to the clipboard
+	 */
+	void OnCopy(ELevelTransformField::Type TransformField);
+
+	/**
+	 * Pastes the specified transform field from the clipboard
+	 */
+	void OnPaste(ELevelTransformField::Type TransformField);
+
+	/**
+	 * Creates a UI action for copying a specified transform field
+	 */
+	FUIAction CreateCopyAction(ELevelTransformField::Type TransformField) const;
+
+	/**
+	 * Creates a UI action for pasting a specified transform field
+	 */
+	FUIAction CreatePasteAction(ELevelTransformField::Type TransformField) const;
+
 private:
 	TWeakPtr<FStreamingLevelCollectionModel>	WorldModel;
 	TSharedPtr<class IPropertyHandle>			LevelPositionProperty;
diff --git a/Engine/Source/Editor/WorldBrowser/WorldBrowser.Build.cs b/Engine/Source/Editor/WorldBrowser/WorldBrowser.Build.cs
--- a/Engine/Source/Editor/WorldBrowser/WorldBrowser.Build.cs
+++ b/Engine/Source/Editor/WorldBrowser/WorldBrowser.Build.cs
@@ -24,6 +24,7 @@ public class WorldBrowser : ModuleRules
      
         PrivateDependencyModuleNames.AddRange(
             new string[] {
+				"ApplicationCore",
                 "AppFramework",
                 "Core", 
                 "CoreUObject",
