Adding "Show All But This" to levels menu, and adding Alt+Click option to eye button to enable "show only this" and "show all but this"
#rnx

diff --git a/Engine/Source/Editor/WorldBrowser/Private/LevelCollectionCommands.h b/Engine/Source/Editor/WorldBrowser/Private/LevelCollectionCommands.h
--- a/Engine/Source/Editor/WorldBrowser/Private/LevelCollectionCommands.h
+++ b/Engine/Source/Editor/WorldBrowser/Private/LevelCollectionCommands.h
@@ -95,7 +95,8 @@ public:
 		UI_COMMAND( World_ShowSelectedLevels, "Show Selected", "Toggles selected levels to a visible state in the viewports", EUserInterfaceActionType::Button, FInputChord() );
 		UI_COMMAND( World_HideSelectedLevels, "Hide Selected", "Toggles selected levels to an invisible state in the viewports", EUserInterfaceActionType::Button, FInputChord() );
 		UI_COMMAND( World_ShowOnlySelectedLevels, "Show Only Selected", "Toggles the selected levels to a visible state; toggles all other levels to an invisible state", EUserInterfaceActionType::Button, FInputChord() );
-		UI_COMMAND( World_ShowAllLevels, "Show All", "Toggles all levels to a visible state in the viewports", EUserInterfaceActionType::Button, FInputChord() );
+		UI_COMMAND(World_ShowAllButSelectedLevels, "Show All But Selected", "Toggles the selected levels to an invisible state; toggles all other levels to a visible state", EUserInterfaceActionType::Button, FInputChord());
+		UI_COMMAND(World_ShowAllLevels, "Show All", "Toggles all levels to a visible state in the viewports", EUserInterfaceActionType::Button, FInputChord());
 		UI_COMMAND( World_HideAllLevels, "Hide All", "Hides all levels to an invisible state in the viewports", EUserInterfaceActionType::Button, FInputChord() );
 
 		//lock
@@ -264,6 +265,9 @@ public:
 	/** Makes selected Levels visible; makes all others invisible */
 	TSharedPtr< FUICommandInfo > World_ShowOnlySelectedLevels;
 
+	/** Makes selected Levels invisible; makes all others visible */
+	TSharedPtr< FUICommandInfo > World_ShowAllButSelectedLevels;
+
 	/** Makes all Levels visible */
 	TSharedPtr< FUICommandInfo > World_ShowAllLevels;
 
diff --git a/Engine/Source/Editor/WorldBrowser/Private/LevelCollectionModel.cpp b/Engine/Source/Editor/WorldBrowser/Private/LevelCollectionModel.cpp
--- a/Engine/Source/Editor/WorldBrowser/Private/LevelCollectionModel.cpp
+++ b/Engine/Source/Editor/WorldBrowser/Private/LevelCollectionModel.cpp
@@ -159,6 +159,10 @@ void FLevelCollectionModel::BindCommands()
 		FExecuteAction::CreateSP( this, &FLevelCollectionModel::ShowOnlySelectedLevels_Executed  ),
 		FCanExecuteAction::CreateSP( this, &FLevelCollectionModel::AreAnySelectedLevelsLoaded ) );
 
+	ActionList.MapAction(Commands.World_ShowAllButSelectedLevels,
+		FExecuteAction::CreateSP(this, &FLevelCollectionModel::ShowAllButSelectedLevels_Executed),
+		FCanExecuteAction::CreateSP(this, &FLevelCollectionModel::AreAnySelectedLevelsLoaded));
+
 	ActionList.MapAction(Commands.World_ShowAllLevels,
 		FExecuteAction::CreateSP(this, &FLevelCollectionModel::ShowAllLevels_Executed));
 	
@@ -472,6 +476,16 @@ void FLevelCollectionModel::ShowLevels(const FLevelModelList& InLevelList)
 	RequestUpdateAllLevels();
 }
 
+void FLevelCollectionModel::ShowOnlySelectedLevels()
+{
+	ShowOnlySelectedLevels_Executed();
+}
+
+void FLevelCollectionModel::ShowAllButSelectedLevels()
+{
+	ShowAllButSelectedLevels_Executed();
+}
+
 void FLevelCollectionModel::UnlockLevels(const FLevelModelList& InLevelList)
 {
 	if (!IsReadOnly())
@@ -1540,6 +1554,18 @@ void FLevelCollectionModel::ShowOnlySelectedLevels_Executed()
 	ShowSelectedLevels_Executed();
 }
 
+void FLevelCollectionModel::ShowAllButSelectedLevels_Executed()
+{
+	//stash off a copy of the original array, as setting visibility can destroy the selection
+	FLevelModelList SelectedLevelsCopy = GetSelectedLevels();
+
+	InvertSelection_Executed();
+	ShowSelectedLevels_Executed();
+	SetSelectedLevels(SelectedLevelsCopy);
+	HideSelectedLevels_Executed();
+}
+
+
 void FLevelCollectionModel::ShowAllLevels_Executed()
 {
 	ShowLevels(GetFilteredLevels());
@@ -1754,6 +1780,7 @@ void FLevelCollectionModel::FillVisibilitySubMenu(FMenuBuilder& InMenuBuilder)
 	InMenuBuilder.AddMenuEntry( Commands.World_ShowSelectedLevels );
 	InMenuBuilder.AddMenuEntry( Commands.World_HideSelectedLevels );
 	InMenuBuilder.AddMenuEntry( Commands.World_ShowOnlySelectedLevels );
+	InMenuBuilder.AddMenuEntry(Commands.World_ShowAllButSelectedLevels);
 	InMenuBuilder.AddMenuEntry( Commands.World_ShowAllLevels );
 	InMenuBuilder.AddMenuEntry( Commands.World_HideAllLevels );
 }
diff --git a/Engine/Source/Editor/WorldBrowser/Private/LevelCollectionModel.h b/Engine/Source/Editor/WorldBrowser/Private/LevelCollectionModel.h
--- a/Engine/Source/Editor/WorldBrowser/Private/LevelCollectionModel.h
+++ b/Engine/Source/Editor/WorldBrowser/Private/LevelCollectionModel.h
@@ -102,6 +102,12 @@ public:
 	/**	Shows level in the world */
 	void ShowLevels(const FLevelModelList& InLevelList);
 
+	/** Toggles the selected levels to a visible state; toggles all other levels to an invisible state */
+	void ShowOnlySelectedLevels();
+
+	/** Toggles the selected levels to an invisible state; toggles all other levels to a visible state */
+	void ShowAllButSelectedLevels();
+
 	/**	Unlocks level in the world */
 	void UnlockLevels(const FLevelModelList& InLevelList);
 	
@@ -408,6 +414,9 @@ protected:
 
 	/** Toggles the selected levels to a visible state; toggles all other levels to an invisible state */
 	void ShowOnlySelectedLevels_Executed();
+	
+	/** Toggles the selected levels to an invisible state; toggles all other levels to a visible state */
+	void ShowAllButSelectedLevels_Executed();
 
 	/** Toggles all levels to a visible state in the viewports */
 	void ShowAllLevels_Executed();
diff --git a/Engine/Source/Editor/WorldBrowser/Private/SWorldHierarchyImpl.cpp b/Engine/Source/Editor/WorldBrowser/Private/SWorldHierarchyImpl.cpp
--- a/Engine/Source/Editor/WorldBrowser/Private/SWorldHierarchyImpl.cpp
+++ b/Engine/Source/Editor/WorldBrowser/Private/SWorldHierarchyImpl.cpp
@@ -725,7 +725,15 @@ void SWorldHierarchyImpl::OnSelectionChanged(const WorldHierarchy::FWorldTreeIte
 
 	for (const WorldHierarchy::FWorldTreeItemPtr& TreeItem : SelectedItems)
 	{
-		SelectedLevels.Append(TreeItem->GetModel());
+		// Folder items should return all child models, but anything else should only return the model for that item
+		if (TreeItem->GetAsFolderTreeItem() != nullptr)
+		{
+			SelectedLevels.Append(TreeItem->GetLevelModels());
+		}
+		else
+		{
+			SelectedLevels.Append(TreeItem->GetModel());
+		}
 	}
 
 	if (!bFoldersOnlyMode)
diff --git a/Engine/Source/Editor/WorldBrowser/Private/SWorldHierarchyItem.cpp b/Engine/Source/Editor/WorldBrowser/Private/SWorldHierarchyItem.cpp
--- a/Engine/Source/Editor/WorldBrowser/Private/SWorldHierarchyItem.cpp
+++ b/Engine/Source/Editor/WorldBrowser/Private/SWorldHierarchyItem.cpp
@@ -20,6 +20,7 @@
 
 #include "WorldTreeItemTypes.h"
 #include "LevelFolders.h"
+#include "Framework/Application/SlateApplication.h"
 
 #define LOCTEXT_NAMESPACE "WorldBrowser"
 
@@ -384,8 +385,21 @@ FSlateColor SWorldHierarchyItem::GetDrawColor() const
 
 FReply SWorldHierarchyItem::OnToggleVisibility()
 {
-	WorldTreeItem->OnToggleVisibility();
-
+	if (FSlateApplication::Get().GetModifierKeys().AreModifersDown(EModifierKey::Alt))
+	{
+		if (WorldTreeItem->IsVisible()) 
+		{
+			WorldTreeItem->OnShowAllButSelected();
+		}
+		else
+		{
+			WorldTreeItem->OnShowOnlySelected();
+		}
+	}
+	else
+	{
+		WorldTreeItem->OnToggleVisibility();
+	}
 	return FReply::Handled();
 }
 
diff --git a/Engine/Source/Editor/WorldBrowser/Private/WorldTreeItemTypes.cpp b/Engine/Source/Editor/WorldBrowser/Private/WorldTreeItemTypes.cpp
--- a/Engine/Source/Editor/WorldBrowser/Private/WorldTreeItemTypes.cpp
+++ b/Engine/Source/Editor/WorldBrowser/Private/WorldTreeItemTypes.cpp
@@ -294,6 +294,30 @@ namespace WorldHierarchy
 		}
 	}
 
+	void FLevelModelTreeItem::OnShowOnlySelected()
+	{
+		if (LevelModel.IsValid())
+		{
+			WorldModel.Pin()->ShowOnlySelectedLevels();
+		}
+	}
+
+	void FLevelModelTreeItem::OnShowAllButSelected()
+	{
+		if (LevelModel.IsValid())
+		{
+			WorldModel.Pin()->ShowAllButSelectedLevels();
+		}
+	}
+
+	void FLevelModelTreeItem::PopulateLevelModelList(FLevelModelList& InModelList)
+	{
+		if (LevelModel.IsValid())
+		{
+			InModelList.Add(LevelModel.Pin());
+		}
+	}
+
 	void FLevelModelTreeItem::SetVisible(bool bVisible)
 	{
 		FLevelModelList LevelModels;
@@ -718,6 +742,38 @@ namespace WorldHierarchy
 		SetVisible(!IsVisible());
 	}
 
+	void FFolderTreeItem::OnShowOnlySelected()
+	{
+		// This can be triggered on the non selected folder, so get its children instead
+		FLevelModelList ChildLevels;
+		for (auto& Child : Children)
+		{
+			Child->PopulateLevelModelList(ChildLevels);
+		}
+		WorldModel.Pin()->SetSelectedLevels(ChildLevels);
+		WorldModel.Pin()->ShowOnlySelectedLevels();
+	}
+
+	void FFolderTreeItem::OnShowAllButSelected()
+	{
+		// This can be triggered on the non selected folder, so get its children instead
+		FLevelModelList ChildLevels;
+		for (auto& Child : Children)
+		{
+			Child->PopulateLevelModelList(ChildLevels);
+		}
+		WorldModel.Pin()->SetSelectedLevels(ChildLevels);
+		WorldModel.Pin()->ShowAllButSelectedLevels();
+	}
+
+	void FFolderTreeItem::PopulateLevelModelList(FLevelModelList& InModelList)
+	{
+		for (auto& Child : Children)
+		{
+			Child->PopulateLevelModelList(InModelList);
+		}
+	}
+
 	void FFolderTreeItem::SetVisible(bool bVisible)
 	{
 		for (auto& Child : Children)
diff --git a/Engine/Source/Editor/WorldBrowser/Private/WorldTreeItemTypes.h b/Engine/Source/Editor/WorldBrowser/Private/WorldTreeItemTypes.h
--- a/Engine/Source/Editor/WorldBrowser/Private/WorldTreeItemTypes.h
+++ b/Engine/Source/Editor/WorldBrowser/Private/WorldTreeItemTypes.h
@@ -72,6 +72,9 @@ namespace WorldHierarchy
 		virtual void SetDrawColor(const FLinearColor& Color) override;
 
 		virtual void OnToggleVisibility() override;
+		virtual void OnShowOnlySelected() override;
+		virtual void OnShowAllButSelected() override;
+		virtual void PopulateLevelModelList(FLevelModelList& InModelList) override;
 		virtual void OnToggleLightingScenario() override;
 		virtual void OnToggleLock() override;
 		virtual void OnSave() override;
@@ -149,6 +152,9 @@ namespace WorldHierarchy
 		virtual bool IsDirty() const override;
 
 		virtual void OnToggleVisibility() override;
+		virtual void OnShowOnlySelected() override;
+		virtual void OnShowAllButSelected() override;
+		virtual void PopulateLevelModelList(FLevelModelList& InModelList) override;
 		virtual void OnToggleLock() override;
 		virtual void OnSave() override;
 		virtual const FSlateBrush* GetHierarchyItemBrush() const override;
diff --git a/Engine/Source/Editor/WorldBrowser/Public/IWorldTreeItem.h b/Engine/Source/Editor/WorldBrowser/Public/IWorldTreeItem.h
--- a/Engine/Source/Editor/WorldBrowser/Public/IWorldTreeItem.h
+++ b/Engine/Source/Editor/WorldBrowser/Public/IWorldTreeItem.h
@@ -52,7 +52,7 @@ namespace WorldHierarchy
 
 			bool bLocked : 1;
 
-			FlagsType() 
+			FlagsType()
 				: bExpanded(1)
 				, bFilteredOut(0)
 				, bChildrenRequiresSort(1)
@@ -65,7 +65,7 @@ namespace WorldHierarchy
 
 		FlagsType Flags;
 
-		DECLARE_DELEGATE( FOnRenameRequest );
+		DECLARE_DELEGATE(FOnRenameRequest);
 
 		FOnRenameRequest RenameRequestEvent;
 
@@ -260,6 +260,9 @@ namespace WorldHierarchy
 		virtual float GetHierarchyItemBrushWidth() const { return 7.0f; }
 
 		virtual void OnToggleVisibility() {}
+		virtual void OnShowOnlySelected() {}
+		virtual void OnShowAllButSelected() {}
+		virtual void PopulateLevelModelList(FLevelModelList& InModelList) {}
 		virtual void OnToggleLightingScenario() {}
 		virtual void OnToggleLock() {}
 		virtual void OnSave() {}
