Adding exclusive toggle for lock in the level menu. Improving behavior when toggling visibility with multiple things selected.
#rnx

diff --git a/Engine/Source/Editor/WorldBrowser/Private/LevelCollectionCommands.h b/Engine/Source/Editor/WorldBrowser/Private/LevelCollectionCommands.h
--- a/Engine/Source/Editor/WorldBrowser/Private/LevelCollectionCommands.h
+++ b/Engine/Source/Editor/WorldBrowser/Private/LevelCollectionCommands.h
@@ -95,13 +95,15 @@ public:
 		UI_COMMAND( World_ShowSelectedLevels, "Show Selected", "Toggles selected levels to a visible state in the viewports", EUserInterfaceActionType::Button, FInputChord() );
 		UI_COMMAND( World_HideSelectedLevels, "Hide Selected", "Toggles selected levels to an invisible state in the viewports", EUserInterfaceActionType::Button, FInputChord() );
 		UI_COMMAND( World_ShowOnlySelectedLevels, "Show Only Selected", "Toggles the selected levels to a visible state; toggles all other levels to an invisible state", EUserInterfaceActionType::Button, FInputChord() );
-		UI_COMMAND(World_ShowAllButSelectedLevels, "Show All But Selected", "Toggles the selected levels to an invisible state; toggles all other levels to a visible state", EUserInterfaceActionType::Button, FInputChord());
-		UI_COMMAND(World_ShowAllLevels, "Show All", "Toggles all levels to a visible state in the viewports", EUserInterfaceActionType::Button, FInputChord());
+		UI_COMMAND( World_ShowAllButSelectedLevels, "Show All But Selected", "Toggles the selected levels to an invisible state; toggles all other levels to a visible state", EUserInterfaceActionType::Button, FInputChord());
+		UI_COMMAND( World_ShowAllLevels, "Show All", "Toggles all levels to a visible state in the viewports", EUserInterfaceActionType::Button, FInputChord());
 		UI_COMMAND( World_HideAllLevels, "Hide All", "Hides all levels to an invisible state in the viewports", EUserInterfaceActionType::Button, FInputChord() );
 
 		//lock
 		UI_COMMAND( World_LockSelectedLevels, "Lock Selected", "Locks selected levels", EUserInterfaceActionType::Button, FInputChord() );
 		UI_COMMAND( World_UnockSelectedLevels, "Unlock Selected", "Unlocks selected levels", EUserInterfaceActionType::Button, FInputChord() );
+		UI_COMMAND(World_LockOnlySelectedLevels, "Lock Only Selected", "Toggles the selected levels to a locked state; toggles all other levels to an unlocked state", EUserInterfaceActionType::Button, FInputChord());
+		UI_COMMAND(World_LockAllButSelectedLevels, "Lock All But Selected", "Toggles the selected levels to an unlocked state; toggles all other levels to a locked state", EUserInterfaceActionType::Button, FInputChord());
 		UI_COMMAND( World_LockAllLevels, "Lock All", "Locks all levels", EUserInterfaceActionType::Button, FInputChord() );
 		UI_COMMAND( World_UnockAllLevels, "Unlock All", "Unlocks all levels", EUserInterfaceActionType::Button, FInputChord() );
 		UI_COMMAND( World_LockReadOnlyLevels, "Lock Read-Only Levels", "Locks all read-only levels", EUserInterfaceActionType::Button, FInputChord() );
@@ -287,6 +289,12 @@ public:
 	/** Unlocks selected levels */
 	TSharedPtr< FUICommandInfo > World_UnockAllLevels;
 
+	/** Makes selected Levels locked; makes all others unlocked */
+	TSharedPtr< FUICommandInfo > World_LockOnlySelectedLevels;
+
+	/** Makes selected Levels unlocked; makes all others locked */
+	TSharedPtr< FUICommandInfo > World_LockAllButSelectedLevels;
+
 	/** Locks all read-only levels */
 	TSharedPtr< FUICommandInfo > World_LockReadOnlyLevels;
 
diff --git a/Engine/Source/Editor/WorldBrowser/Private/LevelCollectionModel.cpp b/Engine/Source/Editor/WorldBrowser/Private/LevelCollectionModel.cpp
--- a/Engine/Source/Editor/WorldBrowser/Private/LevelCollectionModel.cpp
+++ b/Engine/Source/Editor/WorldBrowser/Private/LevelCollectionModel.cpp
@@ -175,6 +175,12 @@ void FLevelCollectionModel::BindCommands()
 	
 	ActionList.MapAction( Commands.World_UnockSelectedLevels,
 		FExecuteAction::CreateSP( this, &FLevelCollectionModel::UnlockSelectedLevels_Executed  ) );
+
+	ActionList.MapAction(Commands.World_LockOnlySelectedLevels,
+		FExecuteAction::CreateSP(this, &FLevelCollectionModel::LockOnlySelectedLevels_Executed));
+
+	ActionList.MapAction(Commands.World_LockAllButSelectedLevels,
+		FExecuteAction::CreateSP(this, &FLevelCollectionModel::LockAllButSelectedLevels_Executed));
 	
 	ActionList.MapAction( Commands.World_LockAllLevels,
 		FExecuteAction::CreateSP( this, &FLevelCollectionModel::LockAllLevels_Executed  ) );
@@ -520,6 +526,16 @@ void FLevelCollectionModel::LockLevels(const FLevelModelList& InLevelList)
 	}
 }
 
+void FLevelCollectionModel::LockOnlySelectedLevels()
+{
+	LockOnlySelectedLevels_Executed();
+}
+
+void FLevelCollectionModel::LockAllButSelectedLevels()
+{
+	LockAllButSelectedLevels_Executed();
+}
+
 void FLevelCollectionModel::SaveLevels(const FLevelModelList& InLevelList)
 {
 	if (IsReadOnly())
@@ -1586,6 +1602,28 @@ void FLevelCollectionModel::UnlockSelectedLevels_Executed()
 	UnlockLevels(GetSelectedLevels());
 }
 
+void FLevelCollectionModel::LockOnlySelectedLevels_Executed()
+{
+	//stash off a copy of the original array, as setting visibility can destroy the selection
+	FLevelModelList SelectedLevelsCopy = GetSelectedLevels();
+
+	InvertSelection_Executed();
+	UnlockSelectedLevels_Executed();
+	SetSelectedLevels(SelectedLevelsCopy);
+	LockSelectedLevels_Executed();
+}
+
+void FLevelCollectionModel::LockAllButSelectedLevels_Executed()
+{
+	//stash off a copy of the original array, as setting visibility can destroy the selection
+	FLevelModelList SelectedLevelsCopy = GetSelectedLevels();
+
+	InvertSelection_Executed();
+	LockSelectedLevels_Executed();
+	SetSelectedLevels(SelectedLevelsCopy);
+	UnlockSelectedLevels_Executed();
+}
+
 void FLevelCollectionModel::LockAllLevels_Executed()
 {
 	if (!IsReadOnly())
@@ -1760,6 +1798,8 @@ void FLevelCollectionModel::FillLockSubMenu(FMenuBuilder& InMenuBuilder)
 
 	InMenuBuilder.AddMenuEntry( Commands.World_LockSelectedLevels );
 	InMenuBuilder.AddMenuEntry( Commands.World_UnockSelectedLevels );
+	InMenuBuilder.AddMenuEntry(Commands.World_LockOnlySelectedLevels);
+	InMenuBuilder.AddMenuEntry(Commands.World_LockAllButSelectedLevels);
 	InMenuBuilder.AddMenuEntry( Commands.World_LockAllLevels );
 	InMenuBuilder.AddMenuEntry( Commands.World_UnockAllLevels );
 
diff --git a/Engine/Source/Editor/WorldBrowser/Private/LevelCollectionModel.h b/Engine/Source/Editor/WorldBrowser/Private/LevelCollectionModel.h
--- a/Engine/Source/Editor/WorldBrowser/Private/LevelCollectionModel.h
+++ b/Engine/Source/Editor/WorldBrowser/Private/LevelCollectionModel.h
@@ -114,6 +114,12 @@ public:
 	/**	Locks level in the world */
 	void LockLevels(const FLevelModelList& InLevelList);
 
+	/** Toggles the selected levels to a locked state; toggles all other levels to an unlocked state */
+	void LockOnlySelectedLevels();
+
+	/** Toggles the selected levels to an unlocked state; toggles all other levels to a locked state */
+	void LockAllButSelectedLevels();
+
 	/**	Saves level to disk */
 	void SaveLevels(const FLevelModelList& InLevelList);
 
@@ -430,6 +436,12 @@ protected:
 	/** Unlocks selected levels */
 	void UnlockSelectedLevels_Executed();
 
+	/** Toggles the selected levels to a locked state; toggles all other levels to an unlocked state */
+	void LockOnlySelectedLevels_Executed();
+
+	/** Toggles the selected levels to an unlocked state; toggles all other levels to a locked state */
+	void LockAllButSelectedLevels_Executed();
+
 	/** Locks all levels */
 	void LockAllLevels_Executed();
 
diff --git a/Engine/Source/Editor/WorldBrowser/Private/SWorldHierarchyItem.cpp b/Engine/Source/Editor/WorldBrowser/Private/SWorldHierarchyItem.cpp
--- a/Engine/Source/Editor/WorldBrowser/Private/SWorldHierarchyItem.cpp
+++ b/Engine/Source/Editor/WorldBrowser/Private/SWorldHierarchyItem.cpp
@@ -411,7 +411,21 @@ FReply SWorldHierarchyItem::OnToggleLightingScenario()
 
 FReply SWorldHierarchyItem::OnToggleLock()
 {
-	WorldTreeItem->OnToggleLock();
+	if (FSlateApplication::Get().GetModifierKeys().AreModifersDown(EModifierKey::Alt))
+	{
+		if (WorldTreeItem->IsLocked())
+		{
+			WorldTreeItem->OnLockAllButSelected();
+		}
+		else
+		{
+			WorldTreeItem->OnLockOnlySelected();
+		}
+	}
+	else
+	{
+		WorldTreeItem->OnToggleLock();
+	}
 	
 	return FReply::Handled();
 }
diff --git a/Engine/Source/Editor/WorldBrowser/Private/WorldTreeItemTypes.cpp b/Engine/Source/Editor/WorldBrowser/Private/WorldTreeItemTypes.cpp
--- a/Engine/Source/Editor/WorldBrowser/Private/WorldTreeItemTypes.cpp
+++ b/Engine/Source/Editor/WorldBrowser/Private/WorldTreeItemTypes.cpp
@@ -296,18 +296,28 @@ namespace WorldHierarchy
 
 	void FLevelModelTreeItem::OnShowOnlySelected()
 	{
-		if (LevelModel.IsValid())
+		FLevelModelList LevelsToChange;
+		PopulateLevelModelList(LevelsToChange);
+		if (GetLevelSelectionFlag())
 		{
-			WorldModel.Pin()->ShowOnlySelectedLevels();
+			FLevelModelList CurrentLevels = WorldModel.Pin()->GetSelectedLevels();
+			LevelsToChange.Append(CurrentLevels);
 		}
+		WorldModel.Pin()->SetSelectedLevels(LevelsToChange);
+		WorldModel.Pin()->ShowOnlySelectedLevels();
 	}
 
 	void FLevelModelTreeItem::OnShowAllButSelected()
 	{
-		if (LevelModel.IsValid())
+		FLevelModelList LevelsToChange;
+		PopulateLevelModelList(LevelsToChange);
+		if (GetLevelSelectionFlag())
 		{
-			WorldModel.Pin()->ShowAllButSelectedLevels();
+			FLevelModelList CurrentLevels = WorldModel.Pin()->GetSelectedLevels();
+			LevelsToChange.Append(CurrentLevels);
 		}
+		WorldModel.Pin()->SetSelectedLevels(LevelsToChange);
+		WorldModel.Pin()->ShowAllButSelectedLevels();
 	}
 
 	void FLevelModelTreeItem::PopulateLevelModelList(FLevelModelList& InModelList)
@@ -354,6 +364,32 @@ namespace WorldHierarchy
 		}
 	}
 
+	void FLevelModelTreeItem::OnLockOnlySelected()
+	{
+		FLevelModelList LevelsToChange;
+		PopulateLevelModelList(LevelsToChange);
+		if (GetLevelSelectionFlag())
+		{
+			FLevelModelList CurrentLevels = WorldModel.Pin()->GetSelectedLevels();
+			LevelsToChange.Append(CurrentLevels);
+		}
+		WorldModel.Pin()->SetSelectedLevels(LevelsToChange);
+		WorldModel.Pin()->LockOnlySelectedLevels();
+	}
+
+	void FLevelModelTreeItem::OnLockAllButSelected()
+	{
+		FLevelModelList LevelsToChange;
+		PopulateLevelModelList(LevelsToChange);
+		if (GetLevelSelectionFlag())
+		{
+			FLevelModelList CurrentLevels = WorldModel.Pin()->GetSelectedLevels();
+			LevelsToChange.Append(CurrentLevels);
+		}
+		WorldModel.Pin()->SetSelectedLevels(LevelsToChange);
+		WorldModel.Pin()->LockAllButSelectedLevels();
+	}
+
 	bool FLevelModelTreeItem::IsLocked() const
 	{
 		return IsLoaded() && LevelModel.Pin()->IsLocked();
@@ -745,24 +781,34 @@ namespace WorldHierarchy
 	void FFolderTreeItem::OnShowOnlySelected()
 	{
 		// This can be triggered on the non selected folder, so get its children instead
-		FLevelModelList ChildLevels;
+		FLevelModelList LevelsToChange;
 		for (auto& Child : Children)
 		{
-			Child->PopulateLevelModelList(ChildLevels);
+			Child->PopulateLevelModelList(LevelsToChange);
 		}
-		WorldModel.Pin()->SetSelectedLevels(ChildLevels);
+		if (GetLevelSelectionFlag())
+		{
+			FLevelModelList CurrentLevels = WorldModel.Pin()->GetSelectedLevels();
+			LevelsToChange.Append(CurrentLevels);
+		}
+		WorldModel.Pin()->SetSelectedLevels(LevelsToChange);
 		WorldModel.Pin()->ShowOnlySelectedLevels();
 	}
 
 	void FFolderTreeItem::OnShowAllButSelected()
 	{
 		// This can be triggered on the non selected folder, so get its children instead
-		FLevelModelList ChildLevels;
+		FLevelModelList LevelsToChange;
 		for (auto& Child : Children)
 		{
-			Child->PopulateLevelModelList(ChildLevels);
+			Child->PopulateLevelModelList(LevelsToChange);
+		}
+		if (GetLevelSelectionFlag())
+		{
+			FLevelModelList CurrentLevels = WorldModel.Pin()->GetSelectedLevels();
+			LevelsToChange.Append(CurrentLevels);
 		}
-		WorldModel.Pin()->SetSelectedLevels(ChildLevels);
+		WorldModel.Pin()->SetSelectedLevels(LevelsToChange);
 		WorldModel.Pin()->ShowAllButSelectedLevels();
 	}
 
@@ -800,6 +846,40 @@ namespace WorldHierarchy
 		SetLocked(!IsLocked());
 	}
 
+	void FFolderTreeItem::OnLockOnlySelected()
+	{
+		// This can be triggered on the non selected folder, so get its children instead
+		FLevelModelList LevelsToChange;
+		for (auto& Child : Children)
+		{
+			Child->PopulateLevelModelList(LevelsToChange);
+		}
+		if (GetLevelSelectionFlag())
+		{
+			FLevelModelList CurrentLevels = WorldModel.Pin()->GetSelectedLevels();
+			LevelsToChange.Append(CurrentLevels);
+		}
+		WorldModel.Pin()->SetSelectedLevels(LevelsToChange);
+		WorldModel.Pin()->LockOnlySelectedLevels();
+	}
+
+	void FFolderTreeItem::OnLockAllButSelected()
+	{
+		// This can be triggered on the non selected folder, so get its children instead
+		FLevelModelList LevelsToChange;
+		for (auto& Child : Children)
+		{
+			Child->PopulateLevelModelList(LevelsToChange);
+		}
+		if (GetLevelSelectionFlag())
+		{
+			FLevelModelList CurrentLevels = WorldModel.Pin()->GetSelectedLevels();
+			LevelsToChange.Append(CurrentLevels);
+		}
+		WorldModel.Pin()->SetSelectedLevels(LevelsToChange);
+		WorldModel.Pin()->LockAllButSelectedLevels();
+	}
+
 	void FFolderTreeItem::SetLocked(bool bLocked)
 	{
 		for (auto& Child : Children)
diff --git a/Engine/Source/Editor/WorldBrowser/Private/WorldTreeItemTypes.h b/Engine/Source/Editor/WorldBrowser/Private/WorldTreeItemTypes.h
--- a/Engine/Source/Editor/WorldBrowser/Private/WorldTreeItemTypes.h
+++ b/Engine/Source/Editor/WorldBrowser/Private/WorldTreeItemTypes.h
@@ -77,6 +77,8 @@ namespace WorldHierarchy
 		virtual void PopulateLevelModelList(FLevelModelList& InModelList) override;
 		virtual void OnToggleLightingScenario() override;
 		virtual void OnToggleLock() override;
+		virtual void OnLockOnlySelected() override;
+		virtual void OnLockAllButSelected() override;
 		virtual void OnSave() override;
 		virtual void OnOpenKismet() override;
 
@@ -156,6 +158,8 @@ namespace WorldHierarchy
 		virtual void OnShowAllButSelected() override;
 		virtual void PopulateLevelModelList(FLevelModelList& InModelList) override;
 		virtual void OnToggleLock() override;
+		virtual void OnLockOnlySelected() override;
+		virtual void OnLockAllButSelected() override;
 		virtual void OnSave() override;
 		virtual const FSlateBrush* GetHierarchyItemBrush() const override;
 		virtual float GetHierarchyItemBrushWidth() const { return 16.0f; }
diff --git a/Engine/Source/Editor/WorldBrowser/Public/IWorldTreeItem.h b/Engine/Source/Editor/WorldBrowser/Public/IWorldTreeItem.h
--- a/Engine/Source/Editor/WorldBrowser/Public/IWorldTreeItem.h
+++ b/Engine/Source/Editor/WorldBrowser/Public/IWorldTreeItem.h
@@ -264,6 +264,8 @@ namespace WorldHierarchy
 		virtual void OnShowAllButSelected() {}
 		virtual void PopulateLevelModelList(FLevelModelList& InModelList) {}
 		virtual void OnToggleLightingScenario() {}
+		virtual void OnLockOnlySelected() {}
+		virtual void OnLockAllButSelected() {}
 		virtual void OnToggleLock() {}
 		virtual void OnSave() {}
 		virtual void OnOpenKismet() {} 
