Add missing return statement to SWorldHierarchyItem::OnMouseButtonDown

diff --git a/Engine/Source/Editor/WorldBrowser/Private/SWorldHierarchyItem.cpp b/Engine/Source/Editor/WorldBrowser/Private/SWorldHierarchyItem.cpp
--- a/Engine/Source/Editor/WorldBrowser/Private/SWorldHierarchyItem.cpp
+++ b/Engine/Source/Editor/WorldBrowser/Private/SWorldHierarchyItem.cpp
@@ -508,6 +508,7 @@ FReply SWorldHierarchyItem::OnMouseButtonDown(const FGeometry& MyGeometry, const
 		}
 
 		Reply.PreventThrottling();
+		return Reply;
 	}
 
 	return FReply::Handled();
