Level UI now reads ShouldBeVisible which takes editor visibility into account as well
#rnx

diff --git a/Engine/Source/Editor/WorldBrowser/Private/LevelModel.cpp b/Engine/Source/Editor/WorldBrowser/Private/LevelModel.cpp
--- a/Engine/Source/Editor/WorldBrowser/Private/LevelModel.cpp
+++ b/Engine/Source/Editor/WorldBrowser/Private/LevelModel.cpp
@@ -194,7 +194,18 @@ bool FLevelModel::IsVisible() const
 	else
 	{
 		ULevel* Level = GetLevelObject();
-		return Level ? FLevelUtils::IsLevelVisible(Level) :  false;
+		if (Level)
+		{
+			if (ULevelStreaming* StreamingLevel = FLevelUtils::FindStreamingLevel(Level))
+			{
+				return StreamingLevel->ShouldBeVisible();
+			}
+			else
+			{
+				return FLevelUtils::IsLevelVisible(Level);
+			}
+		}
+		return false;
 	}
 }
 
