Empty Level Text fixed internally and the user-facing text for the World Outliner

diff --git a/Engine/Source/Editor/WorldBrowser/Private/LevelCollectionCommands.h b/Engine/Source/Editor/WorldBrowser/Private/LevelCollectionCommands.h
--- a/Engine/Source/Editor/WorldBrowser/Private/LevelCollectionCommands.h
+++ b/Engine/Source/Editor/WorldBrowser/Private/LevelCollectionCommands.h
@@ -46,7 +46,7 @@ public:
 		UI_COMMAND( World_SaveSelectedLevelAs, "Save As...", "Save the selected level as...", EUserInterfaceActionType::Button, FInputChord() );
 		UI_COMMAND( World_MigrateSelectedLevels, "Migrate...", "Copies the selected levels and all their dependencies to a different game", EUserInterfaceActionType::Button, FInputChord() );
 		UI_COMMAND( World_MergeSelectedLevels, "Merge...", "Merges the selected levels into a new level, removing the original levels from the persistent", EUserInterfaceActionType::Button, FInputChord() );
-		UI_COMMAND( World_CreateEmptyLevel, "Create New...", "Creates a new empty Level", EUserInterfaceActionType::Button, FInputChord() );
+		UI_COMMAND( World_CreateNewLevel, "Create New...", "Creates a new level, or choose a level template to start from.", EUserInterfaceActionType::Button, FInputChord() );
 		UI_COMMAND( World_AddExistingLevel, "Add Existing...", "Adds an existing level", EUserInterfaceActionType::Button, FInputChord() );
 		UI_COMMAND( World_AddSelectedActorsToNewLevel, "Create New with Selected Actors...", "Adds the actors currently selected in the active viewport to a new Level", EUserInterfaceActionType::Button, FInputChord() );
 		UI_COMMAND(World_RemoveSelectedLevels, "Remove Selected", "Removes selected levels from the base streaming level", EUserInterfaceActionType::Button, FInputChord() );
@@ -167,8 +167,8 @@ public:
 	/** Merges the selected levels into a new level; prompts for save path; removes the levels that were merged */
 	TSharedPtr< FUICommandInfo > World_MergeSelectedLevels;
 
-	/** Creates a new empty level; prompts for save path */
-	TSharedPtr< FUICommandInfo > World_CreateEmptyLevel;
+	/** Creates a new level; prompts for save path */
+	TSharedPtr< FUICommandInfo > World_CreateNewLevel;
 
 	/** Adds an existing streaming level to the persistent; prompts for path */
 	TSharedPtr< FUICommandInfo > World_AddExistingLevel;
diff --git a/Engine/Source/Editor/WorldBrowser/Private/StreamingLevels/StreamingLevelCollectionModel.cpp b/Engine/Source/Editor/WorldBrowser/Private/StreamingLevels/StreamingLevelCollectionModel.cpp
--- a/Engine/Source/Editor/WorldBrowser/Private/StreamingLevels/StreamingLevelCollectionModel.cpp
+++ b/Engine/Source/Editor/WorldBrowser/Private/StreamingLevels/StreamingLevelCollectionModel.cpp
@@ -173,8 +173,8 @@ void FStreamingLevelCollectionModel::BindCommands()
 		FExecuteAction::CreateSP( this, &FStreamingLevelCollectionModel::RemoveInvalidSelectedLevels_Executed ));
 
 	//levels
-	ActionList.MapAction( Commands.World_CreateEmptyLevel,
-		FExecuteAction::CreateSP( this, &FStreamingLevelCollectionModel::CreateEmptyLevel_Executed  ) );
+	ActionList.MapAction( Commands.World_CreateNewLevel,
+		FExecuteAction::CreateSP( this, &FStreamingLevelCollectionModel::CreateNewLevel_Executed  ) );
 	
 	ActionList.MapAction( Commands.World_AddExistingLevel,
 		FExecuteAction::CreateSP( this, &FStreamingLevelCollectionModel::AddExistingLevel_Executed ) );
@@ -377,7 +377,7 @@ void FStreamingLevelCollectionModel::CustomizeFileMainMenu(FMenuBuilder& InMenuB
 			LOCTEXT("LevelsStreamingMethod_Tooltip", "Changes the default streaming method for a new levels"),
 			FNewMenuDelegate::CreateRaw(this, &FStreamingLevelCollectionModel::FillDefaultStreamingMethodSubMenu ) );
 		
-		InMenuBuilder.AddMenuEntry( Commands.World_CreateEmptyLevel );
+		InMenuBuilder.AddMenuEntry( Commands.World_CreateNewLevel );
 		InMenuBuilder.AddMenuEntry( Commands.World_AddExistingLevel );
 		InMenuBuilder.AddMenuEntry( Commands.World_AddSelectedActorsToNewLevel );
 		InMenuBuilder.AddMenuEntry( Commands.World_MergeSelectedLevels );
@@ -414,7 +414,7 @@ const FLevelModelList& FStreamingLevelCollectionModel::GetInvalidSelectedLevels(
 }
 
 //levels
-void FStreamingLevelCollectionModel::CreateEmptyLevel_Executed()
+void FStreamingLevelCollectionModel::CreateNewLevel_Executed()
 {
 	FString TemplateMapPackageName;
 	FNewLevelDialogModule& NewLevelDialogModule = FModuleManager::LoadModuleChecked<FNewLevelDialogModule>("NewLevelDialog");
diff --git a/Engine/Source/Editor/WorldBrowser/Private/StreamingLevels/StreamingLevelCollectionModel.h b/Engine/Source/Editor/WorldBrowser/Private/StreamingLevels/StreamingLevelCollectionModel.h
--- a/Engine/Source/Editor/WorldBrowser/Private/StreamingLevels/StreamingLevelCollectionModel.h
+++ b/Engine/Source/Editor/WorldBrowser/Private/StreamingLevels/StreamingLevelCollectionModel.h
@@ -76,8 +76,8 @@ private:
 	/** Removes selected levels from world */
 	void RemoveInvalidSelectedLevels_Executed();
 	
-	/** Creates a new empty Level; prompts for level save location */
-	void CreateEmptyLevel_Executed();
+	/** Creates a new level; prompts for level save location */
+	void CreateNewLevel_Executed();
 
 	/** Calls AddExistingLevel which adds an existing level; prompts for path */
 	void AddExistingLevel_Executed();
diff --git a/Engine/Source/Editor/WorldBrowser/Private/Tiles/WorldTileCollectionModel.cpp b/Engine/Source/Editor/WorldBrowser/Private/Tiles/WorldTileCollectionModel.cpp
--- a/Engine/Source/Editor/WorldBrowser/Private/Tiles/WorldTileCollectionModel.cpp
+++ b/Engine/Source/Editor/WorldBrowser/Private/Tiles/WorldTileCollectionModel.cpp
@@ -611,7 +611,7 @@ void FWorldTileCollectionModel::CustomizeFileMainMenu(FMenuBuilder& InMenuBuilde
 		
 	InMenuBuilder.BeginSection("LevelsAddLevel");
 	{
-		InMenuBuilder.AddMenuEntry(Commands.World_CreateEmptyLevel);
+		InMenuBuilder.AddMenuEntry(Commands.World_CreateNewLevel);
 		InMenuBuilder.AddMenuEntry(Commands.ImportTiledLandscape);
 	}
 	InMenuBuilder.EndSection();
@@ -749,8 +749,8 @@ void FWorldTileCollectionModel::BindCommands()
 	FUICommandList& ActionList = *CommandList;
 
 	//
-	ActionList.MapAction(Commands.World_CreateEmptyLevel,
-		FExecuteAction::CreateSP(this, &FWorldTileCollectionModel::CreateEmptyLevel_Executed));
+	ActionList.MapAction(Commands.World_CreateNewLevel,
+		FExecuteAction::CreateSP(this, &FWorldTileCollectionModel::CreateNewLevel_Executed));
 
 	ActionList.MapAction(Commands.ClearParentLink,
 		FExecuteAction::CreateSP(this, &FWorldTileCollectionModel::ClearParentLink_Executed),
@@ -1391,7 +1391,7 @@ bool FWorldTileCollectionModel::IsLayerSelected(const FWorldTileLayer& InLayer)
 	return SelectedLayers.Contains(InLayer);
 }
 
-void FWorldTileCollectionModel::CreateEmptyLevel_Executed()
+void FWorldTileCollectionModel::CreateNewLevel_Executed()
 {
 	CreateNewEmptyLevel();
 }
diff --git a/Engine/Source/Editor/WorldBrowser/Private/Tiles/WorldTileCollectionModel.h b/Engine/Source/Editor/WorldBrowser/Private/Tiles/WorldTileCollectionModel.h
--- a/Engine/Source/Editor/WorldBrowser/Private/Tiles/WorldTileCollectionModel.h
+++ b/Engine/Source/Editor/WorldBrowser/Private/Tiles/WorldTileCollectionModel.h
@@ -204,8 +204,8 @@ private:
 	void FillWeightmapsSubMenu(FMenuBuilder& InMenuBuilder) const;
 	
 private:
-	/** Creates a new empty Level; prompts for level save location */
-	void CreateEmptyLevel_Executed();
+	/** Creates a new level; prompts for level save location */
+	void CreateNewLevel_Executed();
 
 	/** Moves world origin to selected level position */
 	void MoveWorldOrigin_Executed();
