FWorldBrowserModule: keep a sharedptr to the world model instead of a weakptr.

diff --git a/Engine/Source/Editor/WorldBrowser/Private/WorldBrowserModule.cpp b/Engine/Source/Editor/WorldBrowser/Private/WorldBrowserModule.cpp
--- a/Engine/Source/Editor/WorldBrowser/Private/WorldBrowserModule.cpp
+++ b/Engine/Source/Editor/WorldBrowser/Private/WorldBrowserModule.cpp
@@ -97,6 +97,9 @@ void FWorldBrowserModule::StartupModule()
 
 void FWorldBrowserModule::ShutdownModule()
 {
+	check(WorldModel.IsUnique());
+	WorldModel.Reset();
+
 	FLevelFolders::Cleanup();
 
 	if (GEngine)
@@ -152,44 +155,40 @@ void FWorldBrowserModule::OnWorldCompositionChanged(UWorld* InWorld)
 
 void FWorldBrowserModule::OnWorldDestroyed(UWorld* InWorld)
 {
-	TSharedPtr<FLevelCollectionModel> SharedWorldModel = WorldModel.Pin();
 	// Is there any editors alive?
-	if (SharedWorldModel.IsValid())
+	if (WorldModel.IsValid())
 	{
-		UWorld* ManagedWorld = SharedWorldModel->GetWorld(/*bEvenIfPendingKill*/true);
+		UWorld* ManagedWorld = WorldModel->GetWorld(/*bEvenIfPendingKill*/true);
 		// Is it our world gets cleaned up?
 		if (ManagedWorld == InWorld)
 		{
 			// Will reset all references to a shared world model
 			OnBrowseWorld.Broadcast(NULL);
 			// So we have to be the last owner of this model
-			check(SharedWorldModel.IsUnique());
+			check(WorldModel.IsUnique());
+			WorldModel.Reset();
 		}
 	}
 }
 
 TSharedPtr<FLevelCollectionModel> FWorldBrowserModule::SharedWorldModel(UWorld* InWorld)
 {
-	TSharedPtr<FLevelCollectionModel> SharedWorldModel = WorldModel.Pin();
-	if (!SharedWorldModel.IsValid() || SharedWorldModel->GetWorld() != InWorld)
+	if (!WorldModel.IsValid() || WorldModel->GetWorld() != InWorld)
 	{
 		if (InWorld)
 		{
 			if (InWorld->WorldComposition)
 			{
-				SharedWorldModel = FWorldTileCollectionModel::Create(InWorld);
+				WorldModel = FWorldTileCollectionModel::Create(InWorld);
 			}
 			else
 			{
-				SharedWorldModel = FStreamingLevelCollectionModel::Create(InWorld);
+				WorldModel = FStreamingLevelCollectionModel::Create(InWorld);
 			}
 		}
-
-		// Hold weak reference to shared world model
-		WorldModel = SharedWorldModel;
 	}
 	
-	return SharedWorldModel;
+	return WorldModel;
 }
 
 #undef LOCTEXT_NAMESPACE
diff --git a/Engine/Source/Editor/WorldBrowser/Public/WorldBrowserModule.h b/Engine/Source/Editor/WorldBrowser/Public/WorldBrowserModule.h
--- a/Engine/Source/Editor/WorldBrowser/Public/WorldBrowserModule.h
+++ b/Engine/Source/Editor/WorldBrowser/Public/WorldBrowserModule.h
@@ -67,7 +67,7 @@ private:
 	void SetCurrentSublevel(TSharedPtr<class FLevelModel> InLevelModel);
 				
 private:
-	TWeakPtr<class FLevelCollectionModel>	WorldModel;
+	TSharedPtr<class FLevelCollectionModel> WorldModel;
 
 	/** Extender for the level menu */
 	FLevelEditorModule::FLevelEditorMenuExtender LevelMenuExtender;
diff --git a/Engine/Source/Runtime/Launch/Private/LaunchEngineLoop.cpp b/Engine/Source/Runtime/Launch/Private/LaunchEngineLoop.cpp
--- a/Engine/Source/Runtime/Launch/Private/LaunchEngineLoop.cpp
+++ b/Engine/Source/Runtime/Launch/Private/LaunchEngineLoop.cpp
@@ -3877,6 +3877,7 @@ void FEngineLoop::Exit()
 	FModuleManager::Get().UnloadModule("AssetTools", true);
 
 #endif // WITH_EDITOR
+	FModuleManager::Get().UnloadModule("WorldBrowser", true);
 	FModuleManager::Get().UnloadModule("AssetRegistry", true);
 
 #if !PLATFORM_ANDROID || PLATFORM_LUMIN 	// AppPreExit doesn't work on Android
