Fix Enable World Composition not populating Sub-Levels in the Levels Window

#rnx

diff --git a/Engine/Source/Editor/WorldBrowser/Private/WorldBrowserModule.cpp b/Engine/Source/Editor/WorldBrowser/Private/WorldBrowserModule.cpp
--- a/Engine/Source/Editor/WorldBrowser/Private/WorldBrowserModule.cpp
+++ b/Engine/Source/Editor/WorldBrowser/Private/WorldBrowserModule.cpp
@@ -97,8 +97,7 @@ void FWorldBrowserModule::StartupModule()
 
 void FWorldBrowserModule::ShutdownModule()
 {
-	check(WorldModel.IsUnique());
-	WorldModel.Reset();
+	ReleaseWorldModel();
 
 	FLevelFolders::Cleanup();
 
@@ -116,6 +115,16 @@ void FWorldBrowserModule::ShutdownModule()
 	FEditorModeRegistry::Get().UnregisterMode(FBuiltinEditorModes::EM_StreamingLevel);
 }
 
+void FWorldBrowserModule::ReleaseWorldModel()
+{
+	if (WorldModel)
+	{
+		// So we have to be the last owner of this model
+		check(WorldModel.IsUnique());
+		WorldModel.Reset();
+	}
+}
+
 TSharedRef<SWidget> FWorldBrowserModule::CreateWorldBrowserHierarchy()
 {
 	UWorld* EditorWorld = GEditor->GetEditorWorldContext().World();
@@ -149,6 +158,8 @@ void FWorldBrowserModule::OnWorldCompositionChanged(UWorld* InWorld)
 		InWorld->WorldType == EWorldType::Editor)
 	{
 		OnBrowseWorld.Broadcast(NULL);
+		// Release World Model here so that a new one gets created in ShareWorldModel even if the World is the same.
+		ReleaseWorldModel();
 		OnBrowseWorld.Broadcast(InWorld);
 	}
 }
@@ -164,9 +175,7 @@ void FWorldBrowserModule::OnWorldDestroyed(UWorld* InWorld)
 		{
 			// Will reset all references to a shared world model
 			OnBrowseWorld.Broadcast(NULL);
-			// So we have to be the last owner of this model
-			check(WorldModel.IsUnique());
-			WorldModel.Reset();
+			ReleaseWorldModel();
 		}
 	}
 }
diff --git a/Engine/Source/Editor/WorldBrowser/Public/WorldBrowserModule.h b/Engine/Source/Editor/WorldBrowser/Public/WorldBrowserModule.h
--- a/Engine/Source/Editor/WorldBrowser/Public/WorldBrowserModule.h
+++ b/Engine/Source/Editor/WorldBrowser/Public/WorldBrowserModule.h
@@ -65,7 +65,9 @@ private:
 
 	bool IsCurrentSublevel(TSharedPtr<class FLevelModel> InLevelModel);
 	void SetCurrentSublevel(TSharedPtr<class FLevelModel> InLevelModel);
-				
+			
+	void ReleaseWorldModel();
+
 private:
 	TSharedPtr<class FLevelCollectionModel> WorldModel;
 
