Fix Level LOD was including Foliage now that MergeMesh doesn't exclude ISM anymore.

#rnx

diff --git a/Engine/Source/Editor/WorldBrowser/Private/Tiles/WorldTileCollectionModel.cpp b/Engine/Source/Editor/WorldBrowser/Private/Tiles/WorldTileCollectionModel.cpp
--- a/Engine/Source/Editor/WorldBrowser/Private/Tiles/WorldTileCollectionModel.cpp
+++ b/Engine/Source/Editor/WorldBrowser/Private/Tiles/WorldTileCollectionModel.cpp
@@ -1956,7 +1956,7 @@ bool FWorldTileCollectionModel::GenerateLODLevels(FLevelModelList InLevelList, i
 				{
 					LandscapeActors.Add(LandscapeProxy);
 				}
-				else 
+				else if(!Actor->IsA<AInstancedFoliageActor>()) // Exclude Foliage from merged mesh
 				{
 					Actors.Add(Actor);
 				}
