Landscape tiled import
- Allow importing landscapes even when there are warnings in the import format and make sure the warning message is displayed along with the import size in the import dialog.
- When validating if we can reimport landscape tiles, take the file type into account (original code was assuming RAW, which made it impossible to reimport any other formats).

diff --git a/Engine/Source/Editor/WorldBrowser/Private/Tiles/STiledLandscapeImportDlg.cpp b/Engine/Source/Editor/WorldBrowser/Private/Tiles/STiledLandscapeImportDlg.cpp
--- a/Engine/Source/Editor/WorldBrowser/Private/Tiles/STiledLandscapeImportDlg.cpp
+++ b/Engine/Source/Editor/WorldBrowser/Private/Tiles/STiledLandscapeImportDlg.cpp
@@ -502,8 +502,12 @@ FReply STiledLandcapeImportDlg::OnClickedSelectHeightmapTiles()
 					if (HeightmapInfo.ResultCode != ELandscapeImportResult::Success)
 					{
 						StatusMessage = HeightmapInfo.ErrorMessage;
-						bValidTiles = false;
-						break;
+
+						if (HeightmapInfo.ResultCode == ELandscapeImportResult::Error)
+						{
+							bValidTiles = false;
+							break;
+						}
 					}
 
 					const FLandscapeFileResolution* FoundSquareResolution =
@@ -705,6 +709,9 @@ void STiledLandcapeImportDlg::GenerateAllPossibleTileConfigurations()
 
 FText STiledLandcapeImportDlg::GetImportSummaryText() const
 {
+	FTextBuilder StatusMessageBuilder;
+	StatusMessageBuilder.AppendLine(StatusMessage);
+
 	if (ImportSettings.HeightmapFileList.Num() && ImportSettings.ComponentsNum > 0)
 	{
 		// Tile information(num, resolution)
@@ -716,15 +723,15 @@ FText STiledLandcapeImportDlg::GetImportSummaryText() const
 		float WidthX = 0.00001f*ImportSettings.Scale3D.X*WidthInTilesX*ImportSettings.SizeX;
 		float WidthY = 0.00001f*ImportSettings.Scale3D.Y*WidthInTilesY*ImportSettings.SizeX;
 		const FString LandscapeSummary = FString::Printf(TEXT("%.3fx%.3f"), WidthX, WidthY);
-	
-		StatusMessage = FText::Format(
+
+		StatusMessageBuilder.AppendLine(FText::Format(
 			LOCTEXT("TiledLandscapeImport_SummaryText", "{0} tiles, {1}km landscape"), 
 			FText::FromString(TilesSummary),
 			FText::FromString(LandscapeSummary)
-			);
+			));
 	}
 	
-	return StatusMessage;
+	return StatusMessageBuilder.ToText();
 }
 
 FText STiledLandcapeImportDlg::GetWeightmapCountText(TSharedPtr<FTiledLandscapeImportSettings::LandscapeLayerSettings> InLayerData) const
diff --git a/Engine/Source/Editor/WorldBrowser/Private/Tiles/WorldTileModel.cpp b/Engine/Source/Editor/WorldBrowser/Private/Tiles/WorldTileModel.cpp
--- a/Engine/Source/Editor/WorldBrowser/Private/Tiles/WorldTileModel.cpp
+++ b/Engine/Source/Editor/WorldBrowser/Private/Tiles/WorldTileModel.cpp
@@ -16,6 +16,8 @@
 #include "Engine/WorldComposition.h"
 #include "GameFramework/WorldSettings.h"
 #include "LandscapeInfo.h"
+#include "LandscapeEditorModule.h"
+#include "LandscapeFileFormatInterface.h"
 #include "LandscapeStreamingProxy.h"
 
 
@@ -334,15 +336,22 @@ bool FWorldTileModel::IsTiledLandscapeBased() const
 	if (IsLandscapeBased() && !GetLandscape()->ReimportHeightmapFilePath.IsEmpty())
 	{
 		// Check if single landscape actor resolution matches heightmap file size
-		IFileManager& FileManager = IFileManager::Get();
-		const int64 ImportFileSize = FileManager.FileSize(*GetLandscape()->ReimportHeightmapFilePath);
-		
-		FIntRect ComponentsRect = GetLandscape()->GetBoundingRect();
-		int64 LandscapeSamples	= (int64)(ComponentsRect.Width()+1)*(ComponentsRect.Height()+1);
-		// Height samples are 2 bytes wide
-		if (LandscapeSamples*2 == ImportFileSize)
+		ILandscapeEditorModule& LandscapeEditorModule = FModuleManager::GetModuleChecked<ILandscapeEditorModule>("LandscapeEditor");
+		const FString TargetExtension = FPaths::GetExtension(GetLandscape()->ReimportHeightmapFilePath, true);
+		const ILandscapeHeightmapFileFormat* HeightmapFormat = LandscapeEditorModule.GetHeightmapFormatByExtension(*TargetExtension);
+
+		FLandscapeHeightmapInfo HeightmapInfo = HeightmapFormat->Validate(*GetLandscape()->ReimportHeightmapFilePath);
+		if (HeightmapInfo.ResultCode != ELandscapeImportResult::Error)
 		{
-			return true;
+			FIntRect ComponentsRect = GetLandscape()->GetBoundingRect();
+
+			for (FLandscapeFileResolution& PossibleResolution: HeightmapInfo.PossibleResolutions)
+			{
+				if ((PossibleResolution.Width == (ComponentsRect.Width() + 1)) && (PossibleResolution.Height == (ComponentsRect.Height() + 1)))
+				{
+					return true;
+				}
+			}
 		}
 	}
 
