PR #5700: UE-7359: Call correct function to deselect actors properly (Contributed by [REDACTED_PERSON])

diff --git a/Engine/Source/Editor/WorldBrowser/Private/LevelModel.cpp b/Engine/Source/Editor/WorldBrowser/Private/LevelModel.cpp
--- a/Engine/Source/Editor/WorldBrowser/Private/LevelModel.cpp
+++ b/Engine/Source/Editor/WorldBrowser/Private/LevelModel.cpp
@@ -660,7 +660,7 @@ void FLevelModel::DeselectAllActors()
 		AActor* CurActor = (*It);
 		if (CurActor)
 		{
-			SelectedActors->Deselect(CurActor);
+			GEditor->SelectActor(CurActor, false, false);
 		}
 	}
 }
