Avoid calling LoadPackage with an empty template package name when creating a new level.

#rnx

diff --git a/Engine/Source/Editor/WorldBrowser/Private/StreamingLevels/StreamingLevelCollectionModel.cpp b/Engine/Source/Editor/WorldBrowser/Private/StreamingLevels/StreamingLevelCollectionModel.cpp
--- a/Engine/Source/Editor/WorldBrowser/Private/StreamingLevels/StreamingLevelCollectionModel.cpp
+++ b/Engine/Source/Editor/WorldBrowser/Private/StreamingLevels/StreamingLevelCollectionModel.cpp
@@ -421,7 +421,7 @@ void FStreamingLevelCollectionModel::CreateNewLevel_Executed()
 	IMainFrameModule& MainFrameModule = FModuleManager::LoadModuleChecked<IMainFrameModule>(TEXT("MainFrame"));
 	if (NewLevelDialogModule.CreateAndShowNewLevelDialog(MainFrameModule.GetParentWindow(), TemplateMapPackageName))
 	{
-		UPackage* TemplatePackage = LoadPackage(nullptr, *TemplateMapPackageName, LOAD_None);
+		UPackage* TemplatePackage = TemplateMapPackageName.Len() ? LoadPackage(nullptr, *TemplateMapPackageName, LOAD_None) : nullptr;
 		UWorld* TemplateWorld = TemplatePackage ? UWorld::FindWorldInPackage(TemplatePackage) : nullptr;
 
 		// Create the new level
