Fixing for loop end statement

diff --git a/Engine/Source/Editor/WorldBrowser/Private/StreamingLevels/StreamingLevelCollectionModel.cpp b/Engine/Source/Editor/WorldBrowser/Private/StreamingLevels/StreamingLevelCollectionModel.cpp
--- a/Engine/Source/Editor/WorldBrowser/Private/StreamingLevels/StreamingLevelCollectionModel.cpp
+++ b/Engine/Source/Editor/WorldBrowser/Private/StreamingLevels/StreamingLevelCollectionModel.cpp
@@ -497,7 +497,7 @@ void FStreamingLevelCollectionModel::FixupInvalidReference_Executed()
 void FStreamingLevelCollectionModel::RemoveInvalidSelectedLevels_Executed()
 {
 	// needs to be an index-based iterator b/c we are removing elements based on it
-	for (int32 LevelIdx = InvalidSelectedLevels.Num() - 1; LevelIdx == 0; LevelIdx--)
+	for (int32 LevelIdx = InvalidSelectedLevels.Num() - 1; LevelIdx <= 0; LevelIdx--)
 	{
 		TSharedPtr<FStreamingLevelModel> TargetModel = StaticCastSharedPtr<FStreamingLevelModel>(InvalidSelectedLevels[LevelIdx]);
 		ULevelStreaming* LevelStreaming = TargetModel->GetLevelStreaming().Get();
