Minor fix for particular case (when some elements maintain their previous visibility) which was not handled in the previous CL7281985.

diff --git a/Engine/Source/Editor/WorldBrowser/Private/LevelModel.cpp b/Engine/Source/Editor/WorldBrowser/Private/LevelModel.cpp
--- a/Engine/Source/Editor/WorldBrowser/Private/LevelModel.cpp
+++ b/Engine/Source/Editor/WorldBrowser/Private/LevelModel.cpp
@@ -257,6 +257,7 @@ void FLevelModel::SetVisible(TArray<FLevelModel*>& LevelModels, const TArray<boo
 	}
 
 	TArray<ULevel*> Levels;
+	TArray<bool> bAreVisibleCleaned; // Required because of the continue in the for loop
 	TArray<bool> bOldAreDirty;
 	for (int32 LevelModelIdx = 0; LevelModelIdx < LevelModels.Num(); ++LevelModelIdx)
 	{
@@ -268,6 +269,7 @@ void FLevelModel::SetVisible(TArray<FLevelModel*>& LevelModels, const TArray<boo
 		}
 		// Otherwise, add Level and IsDirty into each TArray
 		Levels.Add(LevelModel->GetLevelObject());
+		bAreVisibleCleaned.Add(bAreVisible[LevelModelIdx]);
 		bOldAreDirty.Add(LevelModel->IsDirty());
 	}
 	// Don't create unnecessary transactions
@@ -279,7 +281,7 @@ void FLevelModel::SetVisible(TArray<FLevelModel*>& LevelModels, const TArray<boo
 	const FScopedTransaction Transaction(LOCTEXT("ToggleVisibility", "Toggle Level Visibility"));
 
 	// This call hides all owned actors, etc.
-	EditorLevelUtils::SetLevelsVisibility(Levels, bAreVisible, false);
+	EditorLevelUtils::SetLevelsVisibility(Levels, bAreVisibleCleaned, false);
 
 	// Note that Levels.Num() != LevelModels.Num() given the continue in the first for loop
 	for (int32 LevelModelIdx = 0; LevelModelIdx < Levels.Num(); ++LevelModelIdx)
