Fix for crash when repeatedly opening the sublevel menu

diff --git a/Engine/Source/Editor/WorldBrowser/Private/WorldBrowserModule.cpp b/Engine/Source/Editor/WorldBrowser/Private/WorldBrowserModule.cpp
--- a/Engine/Source/Editor/WorldBrowser/Private/WorldBrowserModule.cpp
+++ b/Engine/Source/Editor/WorldBrowser/Private/WorldBrowserModule.cpp
@@ -44,9 +44,7 @@ void FWorldBrowserModule::BuildLevelMenu(FMenuBuilder& MenuBuilder)
 		FLevelModelList ModelList = Model->GetFilteredLevels();
 		for (TSharedPtr<FLevelModel> LevelModel : ModelList)
 		{
-			FUIAction Action(FExecuteAction::CreateRaw(this, &FWorldBrowserModule::SetCurrentSublevel, LevelModel),
-				FCanExecuteAction(),
-				FIsActionChecked::CreateRaw(this, &FWorldBrowserModule::IsCurrentSublevel, LevelModel));
+			FUIAction Action(FExecuteAction::CreateRaw(this, &FWorldBrowserModule::SetCurrentSublevel, LevelModel));
 			MenuBuilder.AddMenuEntry(FText::FromString(LevelModel->GetDisplayName()), FText::GetEmpty(), FSlateIcon(), Action, NAME_None, EUserInterfaceActionType::Button);
 		}
 	}
