Fix for trying to get a world that hasn't been set yet

diff --git a/Engine/Source/Editor/WorldBrowser/Private/LevelModel.cpp b/Engine/Source/Editor/WorldBrowser/Private/LevelModel.cpp
--- a/Engine/Source/Editor/WorldBrowser/Private/LevelModel.cpp
+++ b/Engine/Source/Editor/WorldBrowser/Private/LevelModel.cpp
@@ -360,7 +360,7 @@ void FLevelModel::MakeLevelCurrent()
 	if (!FLevelUtils::IsLevelLocked(Level))
 	{ 
 		// Make current.
-		if (LevelCollectionModel.GetWorld()->SetCurrentLevel(Level))
+		if (LevelCollectionModel.GetWorld() && LevelCollectionModel.GetWorld()->SetCurrentLevel(Level))
 		{
 			FEditorDelegates::NewCurrentLevel.Broadcast();
 				
