Undo //UE4/Dev-Editor/Engine/Source/Editor/WorldBrowser/Private/... changelist 9551411

diff --git a/Engine/Source/Editor/WorldBrowser/Private/LevelModel.cpp b/Engine/Source/Editor/WorldBrowser/Private/LevelModel.cpp
--- a/Engine/Source/Editor/WorldBrowser/Private/LevelModel.cpp
+++ b/Engine/Source/Editor/WorldBrowser/Private/LevelModel.cpp
@@ -360,7 +360,7 @@ void FLevelModel::MakeLevelCurrent()
 	if (!FLevelUtils::IsLevelLocked(Level))
 	{ 
 		// Make current.
-		if (LevelCollectionModel.GetWorld() && LevelCollectionModel.GetWorld()->SetCurrentLevel(Level))
+		if (LevelCollectionModel.GetWorld()->SetCurrentLevel(Level))
 		{
 			FEditorDelegates::NewCurrentLevel.Broadcast();
 				
