Undo //UE4/Dev-Editor/Engine/Source/Editor/WorldBrowser/Private/... changelist 9550808

diff --git a/Engine/Source/Editor/WorldBrowser/Private/WorldBrowserModule.cpp b/Engine/Source/Editor/WorldBrowser/Private/WorldBrowserModule.cpp
--- a/Engine/Source/Editor/WorldBrowser/Private/WorldBrowserModule.cpp
+++ b/Engine/Source/Editor/WorldBrowser/Private/WorldBrowserModule.cpp
@@ -44,7 +44,9 @@ void FWorldBrowserModule::BuildLevelMenu(FMenuBuilder& MenuBuilder)
 		FLevelModelList ModelList = Model->GetFilteredLevels();
 		for (TSharedPtr<FLevelModel> LevelModel : ModelList)
 		{
-			FUIAction Action(FExecuteAction::CreateRaw(this, &FWorldBrowserModule::SetCurrentSublevel, LevelModel));
+			FUIAction Action(FExecuteAction::CreateRaw(this, &FWorldBrowserModule::SetCurrentSublevel, LevelModel),
+				FCanExecuteAction(),
+				FIsActionChecked::CreateRaw(this, &FWorldBrowserModule::IsCurrentSublevel, LevelModel));
 			MenuBuilder.AddMenuEntry(FText::FromString(LevelModel->GetDisplayName()), FText::GetEmpty(), FSlateIcon(), Action, NAME_None, EUserInterfaceActionType::Button);
 		}
 	}
