More fixes for bad const correctness with delegates.

diff --git a/Engine/Plugins/2D/Paper2D/Source/Paper2DEditor/Private/TileMapEditing/TileMapEdModeToolkit.cpp b/Engine/Plugins/2D/Paper2D/Source/Paper2DEditor/Private/TileMapEditing/TileMapEdModeToolkit.cpp
--- a/Engine/Plugins/2D/Paper2D/Source/Paper2DEditor/Private/TileMapEditing/TileMapEdModeToolkit.cpp
+++ b/Engine/Plugins/2D/Paper2D/Source/Paper2DEditor/Private/TileMapEditing/TileMapEdModeToolkit.cpp
@@ -317,7 +317,7 @@ TSharedRef<SWidget> FTileMapEdModeToolkit::BuildToolBar() const
 		//@TODO: TileMapTerrain: Ugly styling
 		FUIAction TerrainTypeDropdownAction;
 		TerrainTypeDropdownAction.IsActionVisibleDelegate = FIsActionButtonVisible::CreateSP(this, &FTileMapEdModeToolkit::DoesSelectedTileSetHaveTerrains);
-		ToolsToolbar.AddComboButton(TerrainTypeDropdownAction, FOnGetContent::CreateSP(this, &FTileMapEdModeToolkit::GenerateTerrainMenu));
+		ToolsToolbar.AddComboButton(TerrainTypeDropdownAction, FOnGetContent::CreateSP(const_cast<FTileMapEdModeToolkit*>(this), &FTileMapEdModeToolkit::GenerateTerrainMenu));
 	}
 
 	return
diff --git a/Engine/Plugins/Compositing/Composure/Source/ComposureLayersEditor/Private/Widgets/SCompElementPreviewPane.cpp b/Engine/Plugins/Compositing/Composure/Source/ComposureLayersEditor/Private/Widgets/SCompElementPreviewPane.cpp
--- a/Engine/Plugins/Compositing/Composure/Source/ComposureLayersEditor/Private/Widgets/SCompElementPreviewPane.cpp
+++ b/Engine/Plugins/Compositing/Composure/Source/ComposureLayersEditor/Private/Widgets/SCompElementPreviewPane.cpp
@@ -259,9 +259,9 @@ TSharedRef<SWidget> SCompElementPreviewPane::GenerateMenu() const
 			FText(),
 			FSlateIcon(),
 			FUIAction(
-				FExecuteAction::CreateSP(this, &SCompElementPreviewPane::RedChannelToggled),
+				FExecuteAction::CreateSP(const_cast<SCompElementPreviewPane*>(this), &SCompElementPreviewPane::RedChannelToggled),
 				FCanExecuteAction(),
-				FGetActionCheckState::CreateSP(this, &SCompElementPreviewPane::GetRedChannel)
+				FGetActionCheckState::CreateSP(const_cast<SCompElementPreviewPane*>(this), &SCompElementPreviewPane::GetRedChannel)
 			),
 			NAME_None,
 			EUserInterfaceActionType::ToggleButton
@@ -271,9 +271,9 @@ TSharedRef<SWidget> SCompElementPreviewPane::GenerateMenu() const
 			FText(),
 			FSlateIcon(),
 			FUIAction(
-				FExecuteAction::CreateSP(this, &SCompElementPreviewPane::GreenChannelToggled),
+				FExecuteAction::CreateSP(const_cast<SCompElementPreviewPane*>(this), &SCompElementPreviewPane::GreenChannelToggled),
 				FCanExecuteAction(),
-				FGetActionCheckState::CreateSP(this, &SCompElementPreviewPane::GetGreenChannel)
+				FGetActionCheckState::CreateSP(const_cast<SCompElementPreviewPane*>(this), &SCompElementPreviewPane::GetGreenChannel)
 			),
 			NAME_None,
 			EUserInterfaceActionType::ToggleButton
@@ -283,9 +283,9 @@ TSharedRef<SWidget> SCompElementPreviewPane::GenerateMenu() const
 			FText(),
 			FSlateIcon(),
 			FUIAction(
-				FExecuteAction::CreateSP(this, &SCompElementPreviewPane::BlueChannelToggled),
+				FExecuteAction::CreateSP(const_cast<SCompElementPreviewPane*>(this), &SCompElementPreviewPane::BlueChannelToggled),
 				FCanExecuteAction(),
-				FGetActionCheckState::CreateSP(this, &SCompElementPreviewPane::GetBlueChannel)
+				FGetActionCheckState::CreateSP(const_cast<SCompElementPreviewPane*>(this), &SCompElementPreviewPane::GetBlueChannel)
 			),
 			NAME_None,
 			EUserInterfaceActionType::ToggleButton
@@ -295,9 +295,9 @@ TSharedRef<SWidget> SCompElementPreviewPane::GenerateMenu() const
 			FText(),
 			FSlateIcon(),
 			FUIAction(
-				FExecuteAction::CreateSP(this, &SCompElementPreviewPane::AlphaChannelToggled),
+				FExecuteAction::CreateSP(const_cast<SCompElementPreviewPane*>(this), &SCompElementPreviewPane::AlphaChannelToggled),
 				FCanExecuteAction(),
-				FGetActionCheckState::CreateSP(this, &SCompElementPreviewPane::GetAlphaChannel)
+				FGetActionCheckState::CreateSP(const_cast<SCompElementPreviewPane*>(this), &SCompElementPreviewPane::GetAlphaChannel)
 			),
 			NAME_None,
 			EUserInterfaceActionType::ToggleButton
diff --git a/Engine/Plugins/Developer/Concert/ConcertFrontend/Source/ConcertFrontend/Private/Widgets/SConcertBrowser.cpp b/Engine/Plugins/Developer/Concert/ConcertFrontend/Source/ConcertFrontend/Private/Widgets/SConcertBrowser.cpp
--- a/Engine/Plugins/Developer/Concert/ConcertFrontend/Source/ConcertFrontend/Private/Widgets/SConcertBrowser.cpp
+++ b/Engine/Plugins/Developer/Concert/ConcertFrontend/Source/ConcertFrontend/Private/Widgets/SConcertBrowser.cpp
@@ -688,7 +688,7 @@ TSharedRef<ITableRow> SConcertBrowser::MakeSessionRowWidget(TSharedPtr<FConcertS
 		ActiveSessionDef.Visibility = MakeAttributeSP(this, &SConcertBrowser::IsActiveSessionVisible, Item->SessionName);
 		ActiveSessionDef.Text = FEditorFontGlyphs::Info_Circle;
 		ActiveSessionDef.ToolTipText = LOCTEXT("ActiveSessionToolTip", "See the current active session");
-		ActiveSessionDef.OnClicked.BindSP(this, &SConcertBrowser::OnClickActiveSession);
+		ActiveSessionDef.OnClicked.BindSP(const_cast<SConcertBrowser*>(this), &SConcertBrowser::OnClickActiveSession);
 		
 		// Resume Session
 		FConcertUIButtonDefinition& ResumeSessionDef = ButtonDefs.AddDefaulted_GetRef();
@@ -696,7 +696,7 @@ TSharedRef<ITableRow> SConcertBrowser::MakeSessionRowWidget(TSharedPtr<FConcertS
 		ResumeSessionDef.Visibility = MakeAttributeSP(this, &SConcertBrowser::IsResumeSessionVisible, Item->SessionName);
 		ResumeSessionDef.Text = FEditorFontGlyphs::Play_Circle;
 		ResumeSessionDef.ToolTipText = LOCTEXT("ResumeSessionToolTip", "Resume receiving updates from this session");
-		ResumeSessionDef.OnClicked.BindSP(this, &SConcertBrowser::OnClickResumeSession);
+		ResumeSessionDef.OnClicked.BindSP(const_cast<SConcertBrowser*>(this), &SConcertBrowser::OnClickResumeSession);
 
 		// Suspend Session
 		FConcertUIButtonDefinition& SuspendSessionDef = ButtonDefs.AddDefaulted_GetRef();
@@ -704,7 +704,7 @@ TSharedRef<ITableRow> SConcertBrowser::MakeSessionRowWidget(TSharedPtr<FConcertS
 		SuspendSessionDef.Visibility = MakeAttributeSP(this, &SConcertBrowser::IsSuspendSessionVisible, Item->SessionName);
 		SuspendSessionDef.Text = FEditorFontGlyphs::Pause_Circle;
 		SuspendSessionDef.ToolTipText = LOCTEXT("SuspendSessionToolTip", "Suspend receiving updates from this session");
-		SuspendSessionDef.OnClicked.BindSP(this, &SConcertBrowser::OnClickSuspendSession);
+		SuspendSessionDef.OnClicked.BindSP(const_cast<SConcertBrowser*>(this), &SConcertBrowser::OnClickSuspendSession);
 
 		// Delete Session
 		FConcertUIButtonDefinition& DeleteSessionDef = ButtonDefs.AddDefaulted_GetRef();
@@ -712,7 +712,7 @@ TSharedRef<ITableRow> SConcertBrowser::MakeSessionRowWidget(TSharedPtr<FConcertS
 		DeleteSessionDef.Visibility = MakeAttributeSP(this, &SConcertBrowser::IsDeleteSessionVisible, Item);
 		DeleteSessionDef.Text = FEditorFontGlyphs::Trash;
 		DeleteSessionDef.ToolTipText = LOCTEXT("DeleteSessionToolTip", "Delete this session");
-		DeleteSessionDef.OnClicked.BindSP(this, &SConcertBrowser::OnClickDeleteSession, Item->SessionName);
+		DeleteSessionDef.OnClicked.BindSP(const_cast<SConcertBrowser*>(this), &SConcertBrowser::OnClickDeleteSession, Item->SessionName);
 
 		// Join Session
 		FConcertUIButtonDefinition& JoinSessionDef = ButtonDefs.AddDefaulted_GetRef();
@@ -720,7 +720,7 @@ TSharedRef<ITableRow> SConcertBrowser::MakeSessionRowWidget(TSharedPtr<FConcertS
 		JoinSessionDef.Visibility = MakeAttributeSP(this, &SConcertBrowser::IsJoinSessionVisible, Item->SessionName);
 		JoinSessionDef.Text = FEditorFontGlyphs::Sign_In;
 		JoinSessionDef.ToolTipText = LOCTEXT("JoinSessionToolTip", "Join this session");
-		JoinSessionDef.OnClicked.BindSP(this, &SConcertBrowser::OnClickJoinSession, Item->SessionName);
+		JoinSessionDef.OnClicked.BindSP(const_cast<SConcertBrowser*>(this), &SConcertBrowser::OnClickJoinSession, Item->SessionName);
 
 		// Leave Session
 		FConcertUIButtonDefinition& LeaveSessionDef = ButtonDefs.AddDefaulted_GetRef();
@@ -728,7 +728,7 @@ TSharedRef<ITableRow> SConcertBrowser::MakeSessionRowWidget(TSharedPtr<FConcertS
 		LeaveSessionDef.Visibility = MakeAttributeSP(this, &SConcertBrowser::IsLeaveSessionVisible, Item->SessionName);
 		LeaveSessionDef.Text = FEditorFontGlyphs::Sign_Out;
 		LeaveSessionDef.ToolTipText = LOCTEXT("LeaveSessionToolTip", "Leave this session");
-		LeaveSessionDef.OnClicked.BindSP(this, &SConcertBrowser::OnClickLeaveSession);
+		LeaveSessionDef.OnClicked.BindSP(const_cast<SConcertBrowser*>(this), &SConcertBrowser::OnClickLeaveSession);
 
 	}
 
diff --git a/Engine/Plugins/Developer/OneSkyLocalizationService/Source/OneSkyLocalizationService/Private/OneSkyLocalizationServiceProvider.cpp b/Engine/Plugins/Developer/OneSkyLocalizationService/Source/OneSkyLocalizationService/Private/OneSkyLocalizationServiceProvider.cpp
--- a/Engine/Plugins/Developer/OneSkyLocalizationService/Source/OneSkyLocalizationService/Private/OneSkyLocalizationServiceProvider.cpp
+++ b/Engine/Plugins/Developer/OneSkyLocalizationService/Source/OneSkyLocalizationService/Private/OneSkyLocalizationServiceProvider.cpp
@@ -483,7 +483,7 @@ void FOneSkyLocalizationServiceProvider::CustomizeTargetToolbar(TSharedRef<FExte
 	const TSharedRef< FUICommandList > CommandList = MakeShareable(new FUICommandList);
 
 	MenuExtender->AddToolBarExtension("LocalizationService", EExtensionHook::First, CommandList,
-		FToolBarExtensionDelegate::CreateRaw(this, &FOneSkyLocalizationServiceProvider::AddTargetToolbarButtons, LocalizationTarget, CommandList));
+		FToolBarExtensionDelegate::CreateRaw(const_cast<FOneSkyLocalizationServiceProvider*>(this), &FOneSkyLocalizationServiceProvider::AddTargetToolbarButtons, LocalizationTarget, CommandList));
 
 }
 
@@ -506,7 +506,7 @@ void FOneSkyLocalizationServiceProvider::CustomizeTargetSetToolbar(TSharedRef<FE
 	const TSharedRef< FUICommandList > CommandList = MakeShareable(new FUICommandList);
 
 	MenuExtender->AddToolBarExtension("LocalizationService", EExtensionHook::First, CommandList,
-		FToolBarExtensionDelegate::CreateRaw(this, &FOneSkyLocalizationServiceProvider::AddTargetSetToolbarButtons, InLocalizationTargetSet, CommandList));
+		FToolBarExtensionDelegate::CreateRaw(const_cast<FOneSkyLocalizationServiceProvider*>(this), &FOneSkyLocalizationServiceProvider::AddTargetSetToolbarButtons, InLocalizationTargetSet, CommandList));
 
 }
 
diff --git a/Engine/Plugins/Editor/AssetManagerEditor/Source/AssetManagerEditor/Private/ReferenceViewer/ReferenceViewerSchema.cpp b/Engine/Plugins/Editor/AssetManagerEditor/Source/AssetManagerEditor/Private/ReferenceViewer/ReferenceViewerSchema.cpp
--- a/Engine/Plugins/Editor/AssetManagerEditor/Source/AssetManagerEditor/Private/ReferenceViewer/ReferenceViewerSchema.cpp
+++ b/Engine/Plugins/Editor/AssetManagerEditor/Source/AssetManagerEditor/Private/ReferenceViewer/ReferenceViewerSchema.cpp
@@ -66,7 +66,7 @@ void UReferenceViewerSchema::GetContextMenuActions(const UEdGraph* CurrentGraph,
 		MenuBuilder->AddSubMenu(
 			NSLOCTEXT("ReferenceViewerSchema", "MakeCollectionWithTitle", "Make Collection with"),
 			NSLOCTEXT("ReferenceViewerSchema", "MakeCollectionWithTooltip", "Makes a collection with either the referencers or dependencies of the selected nodes."),
-			FNewMenuDelegate::CreateUObject(this, &UReferenceViewerSchema::GetMakeCollectionWithSubMenu)
+			FNewMenuDelegate::CreateUObject(const_cast<UReferenceViewerSchema*>(this), &UReferenceViewerSchema::GetMakeCollectionWithSubMenu)
 		);
 	}
 	MenuBuilder->EndSection();
diff --git a/Engine/Plugins/Editor/AssetManagerEditor/Source/AssetManagerEditor/Private/SAssetAuditBrowser.cpp b/Engine/Plugins/Editor/AssetManagerEditor/Source/AssetManagerEditor/Private/SAssetAuditBrowser.cpp
--- a/Engine/Plugins/Editor/AssetManagerEditor/Source/AssetManagerEditor/Private/SAssetAuditBrowser.cpp
+++ b/Engine/Plugins/Editor/AssetManagerEditor/Source/AssetManagerEditor/Private/SAssetAuditBrowser.cpp
@@ -619,7 +619,7 @@ TSharedRef<SWidget> SAssetAuditBrowser::CreateHistoryMenu(bool bInBackHistory) c
 
 			MenuBuilder.AddMenuEntry(DisplayName, DisplayName, FSlateIcon(),
 				FUIAction(
-					FExecuteAction::CreateRaw(this, &SAssetAuditBrowser::GoToHistoryIndex, HistoryIdx)
+					FExecuteAction::CreateRaw(const_cast<SAssetAuditBrowser*>(this), &SAssetAuditBrowser::GoToHistoryIndex, HistoryIdx)
 				),
 				NAME_None, EUserInterfaceActionType::Button);
 			--HistoryIdx;
@@ -634,7 +634,7 @@ TSharedRef<SWidget> SAssetAuditBrowser::CreateHistoryMenu(bool bInBackHistory) c
 
 			MenuBuilder.AddMenuEntry(DisplayName, DisplayName, FSlateIcon(),
 				FUIAction(
-					FExecuteAction::CreateRaw(this, &SAssetAuditBrowser::GoToHistoryIndex, HistoryIdx)
+					FExecuteAction::CreateRaw(const_cast<SAssetAuditBrowser*>(this), &SAssetAuditBrowser::GoToHistoryIndex, HistoryIdx)
 				),
 				NAME_None, EUserInterfaceActionType::Button);
 			++HistoryIdx;
diff --git a/Engine/Plugins/Editor/GameplayTagsEditor/Source/GameplayTagsEditor/Private/SGameplayTagWidget.cpp b/Engine/Plugins/Editor/GameplayTagsEditor/Source/GameplayTagsEditor/Private/SGameplayTagWidget.cpp
--- a/Engine/Plugins/Editor/GameplayTagsEditor/Source/GameplayTagsEditor/Private/SGameplayTagWidget.cpp
+++ b/Engine/Plugins/Editor/GameplayTagsEditor/Source/GameplayTagsEditor/Private/SGameplayTagWidget.cpp
@@ -1358,7 +1358,7 @@ void SGameplayTagWidget::OpenRenameGameplayTagDialog(TSharedPtr<FGameplayTagNode
 	TSharedRef<SRenameGameplayTagDialog> RenameTagDialog =
 		SNew(SRenameGameplayTagDialog)
 		.GameplayTagNode(GameplayTagNode)
-		.OnGameplayTagRenamed(this, &SGameplayTagWidget::OnGameplayTagRenamed);
+		.OnGameplayTagRenamed(const_cast<SGameplayTagWidget*>(this), &SGameplayTagWidget::OnGameplayTagRenamed);
 
 	RenameTagWindow->SetContent(RenameTagDialog);
 
diff --git a/Engine/Plugins/Experimental/ControlRig/Source/ControlRigEditor/Private/ControlRigEditorModule.cpp b/Engine/Plugins/Experimental/ControlRig/Source/ControlRigEditor/Private/ControlRigEditorModule.cpp
--- a/Engine/Plugins/Experimental/ControlRig/Source/ControlRigEditor/Private/ControlRigEditorModule.cpp
+++ b/Engine/Plugins/Experimental/ControlRig/Source/ControlRigEditor/Private/ControlRigEditorModule.cpp
@@ -835,7 +835,7 @@ void FControlRigEditorModule::GetContextMenuActions(const UControlRigGraphNode*
 					LOCTEXT("ClearArray", "Clear"),
 					LOCTEXT("ClearArray_Tooltip", "Clear this array of all of its entries"),
 					FSlateIcon(),
-					FUIAction(FExecuteAction::CreateUObject(Node, &UControlRigGraphNode::HandleClearArray, Context.Pin->PinName.ToString())));
+					FUIAction(FExecuteAction::CreateUObject(const_cast<UControlRigGraphNode*>(Node), &UControlRigGraphNode::HandleClearArray, Context.Pin->PinName.ToString())));
 
 				Context.MenuBuilder->EndSection();
 			}
@@ -851,13 +851,13 @@ void FControlRigEditorModule::GetContextMenuActions(const UControlRigGraphNode*
 					LOCTEXT("RemoveArrayElement", "Remove"),
 					LOCTEXT("RemoveArrayElement_Tooltip", "Remove this array element"),
 					FSlateIcon(),
-					FUIAction(FExecuteAction::CreateUObject(Node, &UControlRigGraphNode::HandleRemoveArrayElement, Context.Pin->PinName.ToString())));
+					FUIAction(FExecuteAction::CreateUObject(const_cast<UControlRigGraphNode*>(Node), &UControlRigGraphNode::HandleRemoveArrayElement, Context.Pin->PinName.ToString())));
 
 				Context.MenuBuilder->AddMenuEntry(
 					LOCTEXT("InsertArrayElement", "Insert"),
 					LOCTEXT("InsertArrayElement_Tooltip", "Insert an array element after this one"),
 					FSlateIcon(),
-					FUIAction(FExecuteAction::CreateUObject(Node, &UControlRigGraphNode::HandleInsertArrayElement, Context.Pin->PinName.ToString())));
+					FUIAction(FExecuteAction::CreateUObject(const_cast<UControlRigGraphNode*>(Node), &UControlRigGraphNode::HandleInsertArrayElement, Context.Pin->PinName.ToString())));
 
 				Context.MenuBuilder->EndSection();
 			}
diff --git a/Engine/Plugins/FX/Niagara/Source/NiagaraEditor/Private/Customizations/NiagaraEventScriptPropertiesCustomization.cpp b/Engine/Plugins/FX/Niagara/Source/NiagaraEditor/Private/Customizations/NiagaraEventScriptPropertiesCustomization.cpp
--- a/Engine/Plugins/FX/Niagara/Source/NiagaraEditor/Private/Customizations/NiagaraEventScriptPropertiesCustomization.cpp
+++ b/Engine/Plugins/FX/Niagara/Source/NiagaraEditor/Private/Customizations/NiagaraEventScriptPropertiesCustomization.cpp
@@ -339,9 +339,9 @@ TSharedRef<SWidget> FNiagaraEventScriptPropertiesCustomization::OnGetMenuContent
 			SNew(SBox)
 			[
 				SNew(SGraphActionMenu)
-				.OnActionSelected(this, &FNiagaraEventScriptPropertiesCustomization::OnActionSelected)
-				.OnCreateWidgetForAction(SGraphActionMenu::FOnCreateWidgetForAction::CreateSP(this, &FNiagaraEventScriptPropertiesCustomization::OnCreateWidgetForAction))
-				.OnCollectAllActions(this, &FNiagaraEventScriptPropertiesCustomization::CollectAllActions)
+				.OnActionSelected(const_cast<FNiagaraEventScriptPropertiesCustomization*>(this), &FNiagaraEventScriptPropertiesCustomization::OnActionSelected)
+				.OnCreateWidgetForAction(SGraphActionMenu::FOnCreateWidgetForAction::CreateSP(const_cast<FNiagaraEventScriptPropertiesCustomization*>(this), &FNiagaraEventScriptPropertiesCustomization::OnCreateWidgetForAction))
+				.OnCollectAllActions(const_cast<FNiagaraEventScriptPropertiesCustomization*>(this), &FNiagaraEventScriptPropertiesCustomization::CollectAllActions)
 				.AutoExpandActionMenu(false)
 				.ShowFilterTextBox(true)
 			]
diff --git a/Engine/Plugins/FX/Niagara/Source/NiagaraEditor/Private/Customizations/NiagaraTypeCustomizations.cpp b/Engine/Plugins/FX/Niagara/Source/NiagaraEditor/Private/Customizations/NiagaraTypeCustomizations.cpp
--- a/Engine/Plugins/FX/Niagara/Source/NiagaraEditor/Private/Customizations/NiagaraTypeCustomizations.cpp
+++ b/Engine/Plugins/FX/Niagara/Source/NiagaraEditor/Private/Customizations/NiagaraTypeCustomizations.cpp
@@ -142,9 +142,9 @@ TSharedRef<SWidget> FNiagaraVariableAttributeBindingCustomization::OnGetMenuCont
 			SNew(SBox)
 			[
 				SNew(SGraphActionMenu)
-				.OnActionSelected(this, &FNiagaraVariableAttributeBindingCustomization::OnActionSelected)
-				.OnCreateWidgetForAction(SGraphActionMenu::FOnCreateWidgetForAction::CreateSP(this, &FNiagaraVariableAttributeBindingCustomization::OnCreateWidgetForAction))
-				.OnCollectAllActions(this, &FNiagaraVariableAttributeBindingCustomization::CollectAllActions)
+				.OnActionSelected(const_cast<FNiagaraVariableAttributeBindingCustomization*>(this), &FNiagaraVariableAttributeBindingCustomization::OnActionSelected)
+				.OnCreateWidgetForAction(SGraphActionMenu::FOnCreateWidgetForAction::CreateSP(const_cast<FNiagaraVariableAttributeBindingCustomization*>(this), &FNiagaraVariableAttributeBindingCustomization::OnCreateWidgetForAction))
+				.OnCollectAllActions(const_cast<FNiagaraVariableAttributeBindingCustomization*>(this), &FNiagaraVariableAttributeBindingCustomization::CollectAllActions)
 				.AutoExpandActionMenu(false)
 				.ShowFilterTextBox(true)
 			]
diff --git a/Engine/Plugins/FX/Niagara/Source/NiagaraEditor/Private/NiagaraNodeOutput.cpp b/Engine/Plugins/FX/Niagara/Source/NiagaraEditor/Private/NiagaraNodeOutput.cpp
--- a/Engine/Plugins/FX/Niagara/Source/NiagaraEditor/Private/NiagaraNodeOutput.cpp
+++ b/Engine/Plugins/FX/Niagara/Source/NiagaraEditor/Private/NiagaraNodeOutput.cpp
@@ -90,7 +90,7 @@ void UNiagaraNodeOutput::GetContextMenuActions(const FGraphNodeContextMenuBuilde
 				[
 					SNew(SEditableTextBox)
 					.Text_UObject(this, &UNiagaraNodeOutput::GetPinNameText, Pin)
-					.OnTextCommitted_UObject(this, &UNiagaraNodeOutput::PinNameTextCommitted, Pin)
+					.OnTextCommitted_UObject(const_cast<UNiagaraNodeOutput*>(this), &UNiagaraNodeOutput::PinNameTextCommitted, Pin)
 				];
 			Context.MenuBuilder->AddWidget(RenameWidget, LOCTEXT("NameMenuItem", "Name"));
 		}
diff --git a/Engine/Plugins/FX/Niagara/Source/NiagaraEditor/Private/NiagaraNodeParameterMapGet.cpp b/Engine/Plugins/FX/Niagara/Source/NiagaraEditor/Private/NiagaraNodeParameterMapGet.cpp
--- a/Engine/Plugins/FX/Niagara/Source/NiagaraEditor/Private/NiagaraNodeParameterMapGet.cpp
+++ b/Engine/Plugins/FX/Niagara/Source/NiagaraEditor/Private/NiagaraNodeParameterMapGet.cpp
@@ -538,7 +538,7 @@ void UNiagaraNodeParameterMapGet::GetContextMenuActions(const FGraphNodeContextM
 				[
 					SNew(SEditableTextBox)
 					.Text_UObject(this, &UNiagaraNodeParameterMapBase::GetPinDescriptionText, Pin)
-					.OnTextCommitted_UObject(this, &UNiagaraNodeParameterMapBase::PinDescriptionTextCommitted, Pin)
+					.OnTextCommitted_UObject(const_cast<UNiagaraNodeParameterMapGet*>(this), &UNiagaraNodeParameterMapBase::PinDescriptionTextCommitted, Pin)
 				];
 			Context.MenuBuilder->AddWidget(RenameWidget, LOCTEXT("DescMenuItem", "Description"));
 			Context.MenuBuilder->EndSection();
diff --git a/Engine/Plugins/FX/Niagara/Source/NiagaraEditor/Private/NiagaraNodeParameterMapSet.cpp b/Engine/Plugins/FX/Niagara/Source/NiagaraEditor/Private/NiagaraNodeParameterMapSet.cpp
--- a/Engine/Plugins/FX/Niagara/Source/NiagaraEditor/Private/NiagaraNodeParameterMapSet.cpp
+++ b/Engine/Plugins/FX/Niagara/Source/NiagaraEditor/Private/NiagaraNodeParameterMapSet.cpp
@@ -271,7 +271,7 @@ void UNiagaraNodeParameterMapSet::GetContextMenuActions(const FGraphNodeContextM
 				[
 					SNew(SEditableTextBox)
 					.Text_UObject(this, &UNiagaraNodeParameterMapBase::GetPinDescriptionText, Pin)
-					.OnTextCommitted_UObject(this, &UNiagaraNodeParameterMapBase::PinDescriptionTextCommitted, Pin)
+					.OnTextCommitted_UObject(const_cast<UNiagaraNodeParameterMapSet*>(this), &UNiagaraNodeParameterMapBase::PinDescriptionTextCommitted, Pin)
 				];
 			Context.MenuBuilder->AddWidget(RenameWidget, LOCTEXT("DescMenuItem", "Description"));
 
diff --git a/Engine/Plugins/FX/Niagara/Source/NiagaraEditor/Private/NiagaraNodeWithDynamicPins.cpp b/Engine/Plugins/FX/Niagara/Source/NiagaraEditor/Private/NiagaraNodeWithDynamicPins.cpp
--- a/Engine/Plugins/FX/Niagara/Source/NiagaraEditor/Private/NiagaraNodeWithDynamicPins.cpp
+++ b/Engine/Plugins/FX/Niagara/Source/NiagaraEditor/Private/NiagaraNodeWithDynamicPins.cpp
@@ -177,7 +177,7 @@ void UNiagaraNodeWithDynamicPins::GetContextMenuActions(const FGraphNodeContextM
 				[
 					SNew(SEditableTextBox)
 					.Text_UObject(this, &UNiagaraNodeWithDynamicPins::GetPinNameText, Pin)
-					.OnTextCommitted_UObject(this, &UNiagaraNodeWithDynamicPins::PinNameTextCommitted, Pin)
+					.OnTextCommitted_UObject(const_cast<UNiagaraNodeWithDynamicPins*>(this), &UNiagaraNodeWithDynamicPins::PinNameTextCommitted, Pin)
 				];
 			Context.MenuBuilder->AddWidget(RenameWidget, LOCTEXT("NameMenuItem", "Name"));
 		}
diff --git a/Engine/Plugins/FX/Niagara/Source/NiagaraEditor/Private/TypeEditorUtilities/NiagaraDataInterfaceCurveTypeEditorUtilities.cpp b/Engine/Plugins/FX/Niagara/Source/NiagaraEditor/Private/TypeEditorUtilities/NiagaraDataInterfaceCurveTypeEditorUtilities.cpp
--- a/Engine/Plugins/FX/Niagara/Source/NiagaraEditor/Private/TypeEditorUtilities/NiagaraDataInterfaceCurveTypeEditorUtilities.cpp
+++ b/Engine/Plugins/FX/Niagara/Source/NiagaraEditor/Private/TypeEditorUtilities/NiagaraDataInterfaceCurveTypeEditorUtilities.cpp
@@ -36,7 +36,7 @@ TSharedPtr<SWidget> FNiagaraDataInterfaceCurveTypeEditorUtilitiesBase::CreateDat
 		return SNew(SComboButton)
 			.ButtonStyle(FEditorStyle::Get(), "HoverHintOnly")
 			.ForegroundColor(FSlateColor::UseForeground())
-			.OnGetMenuContent_Raw(this, &FNiagaraDataInterfaceCurveTypeEditorUtilitiesBase::GetImportMenuContent, MakeWeakObjectPtr(CurveDataInterface), DataInterfaceChangedHandler)
+			.OnGetMenuContent_Raw(const_cast<FNiagaraDataInterfaceCurveTypeEditorUtilitiesBase*>(this), &FNiagaraDataInterfaceCurveTypeEditorUtilitiesBase::GetImportMenuContent, MakeWeakObjectPtr(CurveDataInterface), DataInterfaceChangedHandler)
 			.ButtonContent()
 			[
 				SNew(STextBlock)
diff --git a/Engine/Plugins/FX/Niagara/Source/NiagaraEditor/Private/Widgets/SNiagaraSpreadsheetView.cpp b/Engine/Plugins/FX/Niagara/Source/NiagaraEditor/Private/Widgets/SNiagaraSpreadsheetView.cpp
--- a/Engine/Plugins/FX/Niagara/Source/NiagaraEditor/Private/Widgets/SNiagaraSpreadsheetView.cpp
+++ b/Engine/Plugins/FX/Niagara/Source/NiagaraEditor/Private/Widgets/SNiagaraSpreadsheetView.cpp
@@ -687,7 +687,7 @@ TSharedRef<SWidget> SNiagaraSpreadsheetView::OnGetTargetMenuContent() const
 				ComponentName,
 				ComponentTooltip,
 				FSlateIcon(),
-				FUIAction(FExecuteAction::CreateRaw(this, &SNiagaraSpreadsheetView::SetTarget, *It)));
+				FUIAction(FExecuteAction::CreateRaw(const_cast<SNiagaraSpreadsheetView*>(this), &SNiagaraSpreadsheetView::SetTarget, *It)));
 		}
 	}
 
diff --git a/Engine/Source/Developer/AssetTools/Private/AssetRenameManager.cpp b/Engine/Source/Developer/AssetTools/Private/AssetRenameManager.cpp
--- a/Engine/Source/Developer/AssetTools/Private/AssetRenameManager.cpp
+++ b/Engine/Source/Developer/AssetTools/Private/AssetRenameManager.cpp
@@ -1043,7 +1043,7 @@ bool FAssetRenameManager::CheckPackageForSoftObjectReferences(UPackage* Package,
 		// Bind to dirty callback if we aren't already
 		if (!DirtyDelegateHandle.IsValid())
 		{
-			DirtyDelegateHandle = UPackage::PackageMarkedDirtyEvent.AddSP(this, &FAssetRenameManager::OnMarkPackageDirty);
+			DirtyDelegateHandle = UPackage::PackageMarkedDirtyEvent.AddSP(const_cast<FAssetRenameManager*>(this), &FAssetRenameManager::OnMarkPackageDirty);
 		}
 
 		CachedReferences = &CachedSoftReferences.Add(Package->GetFName());
diff --git a/Engine/Source/Developer/AssetTools/Private/AssetTypeActions/AssetTypeActions_SkeletalMesh.cpp b/Engine/Source/Developer/AssetTools/Private/AssetTypeActions/AssetTypeActions_SkeletalMesh.cpp
--- a/Engine/Source/Developer/AssetTools/Private/AssetTypeActions/AssetTypeActions_SkeletalMesh.cpp
+++ b/Engine/Source/Developer/AssetTools/Private/AssetTypeActions/AssetTypeActions_SkeletalMesh.cpp
@@ -487,7 +487,7 @@ void FAssetTypeActions_SkeletalMesh::FillCreateMenu(FMenuBuilder& MenuBuilder, T
 		MenuBuilder.AddSubMenu(
 			LOCTEXT("SkeletalMesh_NewPhysicsAssetMenu", "Physics Asset"),
 			LOCTEXT("SkeletalMesh_NewPhysicsAssetMenu_ToolTip", "Options for creating new physics assets from the selected meshes."),
-			FNewMenuDelegate::CreateSP(this, &FAssetTypeActions_SkeletalMesh::GetPhysicsAssetMenu, Meshes));
+			FNewMenuDelegate::CreateSP(const_cast<FAssetTypeActions_SkeletalMesh*>(this), &FAssetTypeActions_SkeletalMesh::GetPhysicsAssetMenu, Meshes));
 	}
 	MenuBuilder.EndSection();
 
@@ -836,7 +836,7 @@ void FAssetTypeActions_SkeletalMesh::FillSkeletonMenu(FMenuBuilder& MenuBuilder,
 		LOCTEXT("SkeletalMesh_NewSkeletonTooltip", "Creates a new skeleton for each of the selected meshes."),
 		FSlateIcon(FEditorStyle::GetStyleSetName(), "AssetIcons.Skeleton"),
 		FUIAction(
-			FExecuteAction::CreateSP(this, &FAssetTypeActions_SkeletalMesh::ExecuteNewSkeleton, Meshes),
+			FExecuteAction::CreateSP(const_cast<FAssetTypeActions_SkeletalMesh*>(this), &FAssetTypeActions_SkeletalMesh::ExecuteNewSkeleton, Meshes),
 			FCanExecuteAction()
 			)
 		);
@@ -846,7 +846,7 @@ void FAssetTypeActions_SkeletalMesh::FillSkeletonMenu(FMenuBuilder& MenuBuilder,
 		LOCTEXT("SkeletalMesh_AssignSkeletonTooltip", "Assigns a skeleton to the selected meshes."),
 		FSlateIcon(FEditorStyle::GetStyleSetName(), "Persona.AssetActions.AssignSkeleton"),
 		FUIAction(
-			FExecuteAction::CreateSP(this, &FAssetTypeActions_SkeletalMesh::ExecuteAssignSkeleton, Meshes),
+			FExecuteAction::CreateSP(const_cast<FAssetTypeActions_SkeletalMesh*>(this), &FAssetTypeActions_SkeletalMesh::ExecuteAssignSkeleton, Meshes),
 			FCanExecuteAction()
 			)
 		);
@@ -856,7 +856,7 @@ void FAssetTypeActions_SkeletalMesh::FillSkeletonMenu(FMenuBuilder& MenuBuilder,
 		LOCTEXT("SkeletalMesh_FindSkeletonTooltip", "Finds the skeleton used by the selected meshes in the content browser."),
 		FSlateIcon(FEditorStyle::GetStyleSetName(), "Persona.AssetActions.FindSkeleton"),
 		FUIAction(
-			FExecuteAction::CreateSP(this, &FAssetTypeActions_SkeletalMesh::ExecuteFindSkeleton, Meshes),
+			FExecuteAction::CreateSP(const_cast<FAssetTypeActions_SkeletalMesh*>(this), &FAssetTypeActions_SkeletalMesh::ExecuteFindSkeleton, Meshes),
 			FCanExecuteAction()
 			)
 		);
diff --git a/Engine/Source/Developer/AssetTools/Private/AssetTypeActions/AssetTypeActions_Skeleton.cpp b/Engine/Source/Developer/AssetTools/Private/AssetTypeActions/AssetTypeActions_Skeleton.cpp
--- a/Engine/Source/Developer/AssetTools/Private/AssetTypeActions/AssetTypeActions_Skeleton.cpp
+++ b/Engine/Source/Developer/AssetTools/Private/AssetTypeActions/AssetTypeActions_Skeleton.cpp
@@ -478,7 +478,7 @@ void FAssetTypeActions_Skeleton::FillCreateMenu(FMenuBuilder& MenuBuilder, TArra
 				LOCTEXT("Skeleton_CreateRigTooltip", "Create Rig from this skeleton."),
 				FSlateIcon(),
 				FUIAction(
-				FExecuteAction::CreateSP(this, &FAssetTypeActions_Skeleton::ExecuteCreateRig, Skeletons),
+				FExecuteAction::CreateSP(const_cast<FAssetTypeActions_Skeleton*>(this), &FAssetTypeActions_Skeleton::ExecuteCreateRig, Skeletons),
 				FCanExecuteAction()
 				)
 				);
@@ -488,7 +488,7 @@ void FAssetTypeActions_Skeleton::FillCreateMenu(FMenuBuilder& MenuBuilder, TArra
 
 	TArray<TWeakObjectPtr<UObject>> Objects;
 	Algo::Transform(Skeletons, Objects, [](const TWeakObjectPtr<USkeleton>& Skeleton) { return Skeleton; });
-	AnimationEditorUtils::FillCreateAssetMenu(MenuBuilder, Objects, FAnimAssetCreated::CreateSP(this, &FAssetTypeActions_Skeleton::OnAssetCreated));
+	AnimationEditorUtils::FillCreateAssetMenu(MenuBuilder, Objects, FAnimAssetCreated::CreateSP(const_cast<FAssetTypeActions_Skeleton*>(this), &FAssetTypeActions_Skeleton::OnAssetCreated));
 }
 
 void FAssetTypeActions_Skeleton::OpenAssetEditor( const TArray<UObject*>& InObjects, TSharedPtr<IToolkitHost> EditWithinLevelEditor )
diff --git a/Engine/Source/Developer/LogVisualizer/Private/SVisualLogger.cpp b/Engine/Source/Developer/LogVisualizer/Private/SVisualLogger.cpp
--- a/Engine/Source/Developer/LogVisualizer/Private/SVisualLogger.cpp
+++ b/Engine/Source/Developer/LogVisualizer/Private/SVisualLogger.cpp
@@ -388,7 +388,7 @@ TSharedRef<SDockTab> SVisualLogger::HandleTabManagerSpawnTab(const FSpawnTabArgs
 	}
 	else if (TabIdentifier == MainViewTabId)
 	{
-		TabWidget = SAssignNew(MainView, SVisualLoggerView, CommandList).OnFiltersSearchChanged(this, &SVisualLogger::OnFiltersSearchChanged);
+		TabWidget = SAssignNew(MainView, SVisualLoggerView, CommandList).OnFiltersSearchChanged(const_cast<SVisualLogger*>(this), &SVisualLogger::OnFiltersSearchChanged);
 		AutoSizeTab = false;
 	}
 	else if (TabIdentifier == LogsListTabId)
diff --git a/Engine/Source/Developer/LogVisualizer/Private/SVisualLoggerFilters.cpp b/Engine/Source/Developer/LogVisualizer/Private/SVisualLoggerFilters.cpp
--- a/Engine/Source/Developer/LogVisualizer/Private/SVisualLoggerFilters.cpp
+++ b/Engine/Source/Developer/LogVisualizer/Private/SVisualLoggerFilters.cpp
@@ -249,7 +249,7 @@ void SVisualLoggerFilters::CreateFiltersMenuCategoryForGraph(FMenuBuilder& MenuB
 			FText::Format(LOCTEXT("FilterByTooltipPrefix", "Filter by {0}"), LabelText),
 			FSlateIcon(),
 			FUIAction(
-			FExecuteAction::CreateSP(this, &SVisualLoggerFilters::FilterByTypeClicked, GraphName, DataName),
+			FExecuteAction::CreateSP(const_cast<SVisualLoggerFilters*>(this), &SVisualLoggerFilters::FilterByTypeClicked, GraphName, DataName),
 			FCanExecuteAction(),
 			FIsActionChecked::CreateSP(this, &SVisualLoggerFilters::IsAssetTypeActionsInUse, GraphName, DataName),
 			FIsActionButtonVisible::CreateLambda([this, LabelText]()->bool{return this->GraphsSearchString.Len() == 0 || LabelText.ToString().Find(this->GraphsSearchString) != INDEX_NONE; })),
diff --git a/Engine/Source/Developer/MeshUtilities/Private/MeshUtilities.cpp b/Engine/Source/Developer/MeshUtilities/Private/MeshUtilities.cpp
--- a/Engine/Source/Developer/MeshUtilities/Private/MeshUtilities.cpp
+++ b/Engine/Source/Developer/MeshUtilities/Private/MeshUtilities.cpp
@@ -5257,8 +5257,8 @@ private:
 				if (StaticMeshReductionInterface)
 				{
 					FUIAction UIAction;
-					UIAction.ExecuteAction.BindSP(this, &FMeshSimplifcationSettingsCustomization::OnMeshSimplificationModuleChosen, ModuleName);
-					UIAction.GetActionCheckState.BindSP(this, &FMeshSimplifcationSettingsCustomization::IsMeshSimplificationModuleChosen, ModuleName);
+					UIAction.ExecuteAction.BindSP(const_cast<FMeshSimplifcationSettingsCustomization*>(this), &FMeshSimplifcationSettingsCustomization::OnMeshSimplificationModuleChosen, ModuleName);
+					UIAction.GetActionCheckState.BindSP(const_cast<FMeshSimplifcationSettingsCustomization*>(this), &FMeshSimplifcationSettingsCustomization::IsMeshSimplificationModuleChosen, ModuleName);
 
 					MenuBuilder.AddMenuEntry(FText::FromName(ModuleName), FText::GetEmpty(), FSlateIcon(), UIAction, NAME_None, EUserInterfaceActionType::RadioButton);
 				}
@@ -5269,7 +5269,7 @@ private:
 		
 
 		FUIAction OpenMarketplaceAction;
-		OpenMarketplaceAction.ExecuteAction.BindSP(this, &FMeshSimplifcationSettingsCustomization::OnFindReductionPluginsClicked);
+		OpenMarketplaceAction.ExecuteAction.BindSP(const_cast<FMeshSimplifcationSettingsCustomization*>(this), &FMeshSimplifcationSettingsCustomization::OnFindReductionPluginsClicked);
 		FSlateIcon Icon = FSlateIcon(FEditorStyle::Get().GetStyleSetName(), "LevelEditor.OpenMarketplace.Menu");
 		MenuBuilder.AddMenuEntry( LOCTEXT("FindMoreReductionPluginsLink", "Search the Marketplace"), LOCTEXT("FindMoreReductionPluginsLink_Tooltip", "Opens the Marketplace to find more mesh reduction plugins"), Icon, OpenMarketplaceAction);
 		return MenuBuilder.MakeWidget();
@@ -5377,8 +5377,8 @@ private:
 				if (SkeletalMeshReductionInterface)
 				{
 					FUIAction UIAction;
-					UIAction.ExecuteAction.BindSP(this, &FSkeletalMeshSimplificationSettingsCustomization::OnSkeletalMeshSimplificationModuleChosen, ModuleName);
-					UIAction.GetActionCheckState.BindSP(this, &FSkeletalMeshSimplificationSettingsCustomization::IsSkeletalMeshSimplificationModuleChosen, ModuleName);
+					UIAction.ExecuteAction.BindSP(const_cast<FSkeletalMeshSimplificationSettingsCustomization*>(this), &FSkeletalMeshSimplificationSettingsCustomization::OnSkeletalMeshSimplificationModuleChosen, ModuleName);
+					UIAction.GetActionCheckState.BindSP(const_cast<FSkeletalMeshSimplificationSettingsCustomization*>(this), &FSkeletalMeshSimplificationSettingsCustomization::IsSkeletalMeshSimplificationModuleChosen, ModuleName);
 
 					MenuBuilder.AddMenuEntry(FText::FromName(ModuleName), FText::GetEmpty(), FSlateIcon(), UIAction, NAME_None, EUserInterfaceActionType::RadioButton);
 				}
@@ -5389,7 +5389,7 @@ private:
 
 
 		FUIAction OpenMarketplaceAction;
-		OpenMarketplaceAction.ExecuteAction.BindSP(this, &FSkeletalMeshSimplificationSettingsCustomization::OnFindReductionPluginsClicked);
+		OpenMarketplaceAction.ExecuteAction.BindSP(const_cast<FSkeletalMeshSimplificationSettingsCustomization*>(this), &FSkeletalMeshSimplificationSettingsCustomization::OnFindReductionPluginsClicked);
 		FSlateIcon Icon = FSlateIcon(FEditorStyle::Get().GetStyleSetName(), "LevelEditor.OpenMarketplace.Menu");
 		MenuBuilder.AddMenuEntry(LOCTEXT("FindMoreReductionPluginsLink", "Search the Marketplace"), LOCTEXT("FindMoreReductionPluginsLink_Tooltip", "Opens the Marketplace to find more mesh reduction plugins"), Icon, OpenMarketplaceAction);
 		return MenuBuilder.MakeWidget();
@@ -5497,8 +5497,8 @@ private:
 				if (MeshMergingInterface)
 				{
 					FUIAction UIAction;
-					UIAction.ExecuteAction.BindSP(this, &FProxyLODMeshSimplificationSettingsCustomization::OnProxyLODMeshSimplificationModuleChosen, ModuleName);
-					UIAction.GetActionCheckState.BindSP(this, &FProxyLODMeshSimplificationSettingsCustomization::IsProxyLODMeshSimplificationModuleChosen, ModuleName);
+					UIAction.ExecuteAction.BindSP(const_cast<FProxyLODMeshSimplificationSettingsCustomization*>(this), &FProxyLODMeshSimplificationSettingsCustomization::OnProxyLODMeshSimplificationModuleChosen, ModuleName);
+					UIAction.GetActionCheckState.BindSP(const_cast<FProxyLODMeshSimplificationSettingsCustomization*>(this), &FProxyLODMeshSimplificationSettingsCustomization::IsProxyLODMeshSimplificationModuleChosen, ModuleName);
 
 					MenuBuilder.AddMenuEntry(FText::FromName(ModuleName), FText::GetEmpty(), FSlateIcon(), UIAction, NAME_None, EUserInterfaceActionType::RadioButton);
 				}
@@ -5509,7 +5509,7 @@ private:
 
 
 		FUIAction OpenMarketplaceAction;
-		OpenMarketplaceAction.ExecuteAction.BindSP(this, &FProxyLODMeshSimplificationSettingsCustomization::OnFindReductionPluginsClicked);
+		OpenMarketplaceAction.ExecuteAction.BindSP(const_cast<FProxyLODMeshSimplificationSettingsCustomization*>(this), &FProxyLODMeshSimplificationSettingsCustomization::OnFindReductionPluginsClicked);
 		FSlateIcon Icon = FSlateIcon(FEditorStyle::Get().GetStyleSetName(), "LevelEditor.OpenMarketplace.Menu");
 		MenuBuilder.AddMenuEntry(LOCTEXT("FindMoreReductionPluginsLink", "Search the Marketplace"), LOCTEXT("FindMoreReductionPluginsLink_Tooltip", "Opens the Marketplace to find more mesh reduction plugins"), Icon, OpenMarketplaceAction);
 		return MenuBuilder.MakeWidget();
diff --git a/Engine/Source/Developer/Profiler/Private/ProfilerCommands.cpp b/Engine/Source/Developer/Profiler/Private/ProfilerCommands.cpp
--- a/Engine/Source/Developer/Profiler/Private/ProfilerCommands.cpp
+++ b/Engine/Source/Developer/Profiler/Private/ProfilerCommands.cpp
@@ -103,7 +103,7 @@ void FProfilerActionManager::Map_ToggleDataPreview_Global()
 const FUIAction FProfilerActionManager::ToggleDataPreview_Custom( const FGuid SessionInstanceID ) const
 {
 	FUIAction UIAction;
-	UIAction.ExecuteAction = FExecuteAction::CreateRaw( this, &FProfilerActionManager::ToggleDataPreview_Execute, SessionInstanceID );
+	UIAction.ExecuteAction = FExecuteAction::CreateRaw( const_cast<FProfilerActionManager*>(this), &FProfilerActionManager::ToggleDataPreview_Execute, SessionInstanceID );
 	UIAction.CanExecuteAction = FCanExecuteAction::CreateRaw( this, &FProfilerActionManager::ToggleDataPreview_CanExecute, SessionInstanceID );
 	UIAction.GetActionCheckState = FGetActionCheckState::CreateRaw( this, &FProfilerActionManager::ToggleDataPreview_GetCheckState, SessionInstanceID );
 	return UIAction;
@@ -338,7 +338,7 @@ void FProfilerActionManager::Map_OpenSettings_Global()
 const FUIAction FProfilerActionManager::OpenSettings_Custom() const
 {
 	FUIAction UIAction;
-	UIAction.ExecuteAction = FExecuteAction::CreateRaw( this, &FProfilerActionManager::OpenSettings_Execute );
+	UIAction.ExecuteAction = FExecuteAction::CreateRaw( const_cast<FProfilerActionManager*>(this), &FProfilerActionManager::OpenSettings_Execute );
 	UIAction.CanExecuteAction = FCanExecuteAction::CreateRaw( this, &FProfilerActionManager::OpenSettings_CanExecute );
 	return UIAction;
 }
diff --git a/Engine/Source/Developer/Profiler/Private/Widgets/SEventGraph.cpp b/Engine/Source/Developer/Profiler/Private/Widgets/SEventGraph.cpp
--- a/Engine/Source/Developer/Profiler/Private/Widgets/SEventGraph.cpp
+++ b/Engine/Source/Developer/Profiler/Private/Widgets/SEventGraph.cpp
@@ -1840,7 +1840,7 @@ TSharedPtr<SWidget> SEventGraph::EventGraph_GetMenuContent() const
 		( 
 			LOCTEXT("ContextMenu_Header_Misc_Sort", "Sort By"), 
 			LOCTEXT("ContextMenu_Header_Misc_Sort_Desc", "Sort by column"), 
-			FNewMenuDelegate::CreateSP( this, &SEventGraph::EventGraph_BuildSortByMenu ),
+			FNewMenuDelegate::CreateSP( const_cast<SEventGraph*>(this), &SEventGraph::EventGraph_BuildSortByMenu ),
 			false, 
 			FSlateIcon(FEditorStyle::GetStyleSetName(), "Profiler.Misc.SortBy")
 		);
diff --git a/Engine/Source/Developer/Profiler/Private/Widgets/SEventGraph.h b/Engine/Source/Developer/Profiler/Private/Widgets/SEventGraph.h
--- a/Engine/Source/Developer/Profiler/Private/Widgets/SEventGraph.h
+++ b/Engine/Source/Developer/Profiler/Private/Widgets/SEventGraph.h
@@ -318,7 +318,7 @@ public:
 	const FUIAction SelectAllFrames_Custom() const
 	{
 		FUIAction UIAction;
-		UIAction.ExecuteAction = FExecuteAction::CreateSP( this, &SEventGraph::SelectAllFrames_Execute );
+		UIAction.ExecuteAction = FExecuteAction::CreateSP( const_cast<SEventGraph*>(this), &SEventGraph::SelectAllFrames_Execute );
 		UIAction.CanExecuteAction = FCanExecuteAction::CreateSP( this, &SEventGraph::SelectAllFrames_CanExecute );
 		return UIAction;
 	}
@@ -425,7 +425,7 @@ public:
 	const FUIAction ShowSelectedEventsInViewMode_Custom( EEventGraphViewModes::Type NewViewMode ) const
 	{
 		FUIAction UIAction;
-		UIAction.ExecuteAction = FExecuteAction::CreateSP( this, &SEventGraph::ShowSelectedEventsInViewMode_Execute, NewViewMode );
+		UIAction.ExecuteAction = FExecuteAction::CreateSP( const_cast<SEventGraph*>(this), &SEventGraph::ShowSelectedEventsInViewMode_Execute, NewViewMode );
 		UIAction.CanExecuteAction = FCanExecuteAction::CreateSP( this, &SEventGraph::ShowSelectedEventsInViewMode_CanExecute, NewViewMode );
 		UIAction.GetActionCheckState = FGetActionCheckState::CreateSP( this, &SEventGraph::ShowSelectedEventsInViewMode_GetCheckState, NewViewMode );
 		return UIAction;
@@ -448,7 +448,7 @@ public:
 	const FUIAction FilterOutByProperty_Custom( const FEventGraphSamplePtr EventPtr, const FName PropertyName, const bool bReset ) const
 	{
 		FUIAction UIAction;
-		UIAction.ExecuteAction = FExecuteAction::CreateSP( this, &SEventGraph::FilterOutByProperty_Execute, EventPtr, PropertyName, bReset );
+		UIAction.ExecuteAction = FExecuteAction::CreateSP( const_cast<SEventGraph*>(this), &SEventGraph::FilterOutByProperty_Execute, EventPtr, PropertyName, bReset );
 		UIAction.CanExecuteAction = FCanExecuteAction::CreateSP( this, &SEventGraph::FilterOutByProperty_CanExecute, EventPtr, PropertyName, bReset );
 		return UIAction;
 	}
@@ -489,7 +489,7 @@ public:
 	const FUIAction HistoryList_GoTo_Custom( int32 StateIndex ) const
 	{
 		FUIAction UIAction;
-		UIAction.ExecuteAction = FExecuteAction::CreateSP( this, &SEventGraph::HistoryList_GoTo_Execute, StateIndex );
+		UIAction.ExecuteAction = FExecuteAction::CreateSP( const_cast<SEventGraph*>(this), &SEventGraph::HistoryList_GoTo_Execute, StateIndex );
 		UIAction.CanExecuteAction = FCanExecuteAction();
 		UIAction.GetActionCheckState = FGetActionCheckState::CreateSP( this, &SEventGraph::HistoryList_GoTo_GetCheckState, StateIndex );
 		return UIAction;
@@ -522,7 +522,7 @@ public:
 	const FUIAction SetExpansionForEvents_Custom( const ESelectedEventTypes::Type SelectedEventType, bool bShouldExpand ) const
 	{
 		FUIAction UIAction;
-		UIAction.ExecuteAction = FExecuteAction::CreateSP( this, &SEventGraph::SetExpansionForEvents_Execute, SelectedEventType, bShouldExpand );
+		UIAction.ExecuteAction = FExecuteAction::CreateSP( const_cast<SEventGraph*>(this), &SEventGraph::SetExpansionForEvents_Execute, SelectedEventType, bShouldExpand );
 		UIAction.CanExecuteAction = FCanExecuteAction::CreateSP( this, &SEventGraph::SetExpansionForEvents_CanExecute, SelectedEventType, bShouldExpand );
 		return UIAction;
 	}
diff --git a/Engine/Source/Developer/SessionFrontend/Private/Widgets/SSessionFrontend.cpp b/Engine/Source/Developer/SessionFrontend/Private/Widgets/SSessionFrontend.cpp
--- a/Engine/Source/Developer/SessionFrontend/Private/Widgets/SSessionFrontend.cpp
+++ b/Engine/Source/Developer/SessionFrontend/Private/Widgets/SSessionFrontend.cpp
@@ -192,14 +192,14 @@ TSharedRef<SDockTab> SSessionFrontend::HandleTabManagerSpawnTab( const FSpawnTab
 		IAutomationControllerManagerPtr AutomationController = AutomationControllerModule.GetAutomationController();
 		IAutomationWindowModule& AutomationWindowModule = FModuleManager::LoadModuleChecked<IAutomationWindowModule>("AutomationWindow");
 
-		AutomationController->OnShutdown().AddSP(this, &SSessionFrontend::HandleAutomationModuleShutdown);
+		AutomationController->OnShutdown().AddSP(const_cast<SSessionFrontend*>(this), &SSessionFrontend::HandleAutomationModuleShutdown);
 
 		TabWidget = AutomationWindowModule.CreateAutomationWindow(
 			AutomationController.ToSharedRef(),
 			SessionManager.ToSharedRef()
 		);
 
-		AutomationWindowModule.OnShutdown().BindSP(this, &SSessionFrontend::HandleAutomationModuleShutdown);
+		AutomationWindowModule.OnShutdown().BindSP(const_cast<SSessionFrontend*>(this), &SSessionFrontend::HandleAutomationModuleShutdown);
 	}
 	else if (TabIdentifier == ProfilerTabId)
 	{
diff --git a/Engine/Source/Editor/AnimGraph/Private/AnimGraphNode_BlendListByEnum.cpp b/Engine/Source/Editor/AnimGraph/Private/AnimGraphNode_BlendListByEnum.cpp
--- a/Engine/Source/Editor/AnimGraph/Private/AnimGraphNode_BlendListByEnum.cpp
+++ b/Engine/Source/Editor/AnimGraph/Private/AnimGraphNode_BlendListByEnum.cpp
@@ -95,7 +95,7 @@ void UAnimGraphNode_BlendListByEnum::GetContextMenuActions(const FGraphNodeConte
 				if (ExposedEnumIndex != INDEX_NONE)
 				{
 					// Offer to remove this specific pin
-					FUIAction Action = FUIAction( FExecuteAction::CreateUObject( this, &UAnimGraphNode_BlendListByEnum::RemovePinFromBlendList, const_cast<UEdGraphPin*>(Context.Pin)) );
+					FUIAction Action = FUIAction( FExecuteAction::CreateUObject( const_cast<UAnimGraphNode_BlendListByEnum*>(this), &UAnimGraphNode_BlendListByEnum::RemovePinFromBlendList, const_cast<UEdGraphPin*>(Context.Pin)) );
 					Context.MenuBuilder->AddMenuEntry( LOCTEXT("RemovePose", "Remove Pose"), FText::GetEmpty(), FSlateIcon(), Action );
 				}
 			}
@@ -117,14 +117,14 @@ void UAnimGraphNode_BlendListByEnum::GetContextMenuActions(const FGraphNodeConte
 					bAddedHeader = true;
 					Context.MenuBuilder->BeginSection("AnimGraphNodeAddElementPin", LOCTEXT("ExposeHeader", "Add pin for element"));
 					{
-						FUIAction Action = FUIAction( FExecuteAction::CreateUObject( this, &UAnimGraphNode_BlendListByEnum::ExposeEnumElementAsPin, ElementName) );
+						FUIAction Action = FUIAction( FExecuteAction::CreateUObject( const_cast<UAnimGraphNode_BlendListByEnum*>(this), &UAnimGraphNode_BlendListByEnum::ExposeEnumElementAsPin, ElementName) );
 						Context.MenuBuilder->AddMenuEntry(PrettyElementName, PrettyElementName, FSlateIcon(), Action);
 					}
 					Context.MenuBuilder->EndSection();
 				}
 				else
 				{
-					FUIAction Action = FUIAction( FExecuteAction::CreateUObject( this, &UAnimGraphNode_BlendListByEnum::ExposeEnumElementAsPin, ElementName) );
+					FUIAction Action = FUIAction( FExecuteAction::CreateUObject( const_cast<UAnimGraphNode_BlendListByEnum*>(this), &UAnimGraphNode_BlendListByEnum::ExposeEnumElementAsPin, ElementName) );
 					Context.MenuBuilder->AddMenuEntry(PrettyElementName, PrettyElementName, FSlateIcon(), Action);
 				}
 			}
diff --git a/Engine/Source/Editor/AnimGraph/Private/AnimGraphNode_ModifyCurve.cpp b/Engine/Source/Editor/AnimGraph/Private/AnimGraphNode_ModifyCurve.cpp
--- a/Engine/Source/Editor/AnimGraph/Private/AnimGraphNode_ModifyCurve.cpp
+++ b/Engine/Source/Editor/AnimGraph/Private/AnimGraphNode_ModifyCurve.cpp
@@ -59,7 +59,7 @@ void UAnimGraphNode_ModifyCurve::GetAddCurveMenuActions(FMenuBuilder& MenuBuilde
 	TArray<FName> CurvesToAdd = GetCurvesToAdd();
 	for (FName CurveName : CurvesToAdd)
 	{
-		FUIAction Action = FUIAction(FExecuteAction::CreateUObject(this, &UAnimGraphNode_ModifyCurve::AddCurvePin, CurveName));
+		FUIAction Action = FUIAction(FExecuteAction::CreateUObject(const_cast<UAnimGraphNode_ModifyCurve*>(this), &UAnimGraphNode_ModifyCurve::AddCurvePin, CurveName));
 		MenuBuilder.AddMenuEntry(FText::FromName(CurveName), FText::GetEmpty(), FSlateIcon(), Action);
 	}
 }
@@ -68,7 +68,7 @@ void UAnimGraphNode_ModifyCurve::GetRemoveCurveMenuActions(FMenuBuilder& MenuBui
 {
 	for (FName CurveName : Node.CurveNames)
 	{
-		FUIAction Action = FUIAction(FExecuteAction::CreateUObject(this, &UAnimGraphNode_ModifyCurve::RemoveCurvePin, CurveName));
+		FUIAction Action = FUIAction(FExecuteAction::CreateUObject(const_cast<UAnimGraphNode_ModifyCurve*>(this), &UAnimGraphNode_ModifyCurve::RemoveCurvePin, CurveName));
 		MenuBuilder.AddMenuEntry(FText::FromName(CurveName), FText::GetEmpty(), FSlateIcon(), Action);
 	}
 }
@@ -92,7 +92,7 @@ void UAnimGraphNode_ModifyCurve::GetContextMenuActions(const FGraphNodeContextMe
 			if (PinPropertyName  == GET_MEMBER_NAME_CHECKED(FAnimNode_ModifyCurve, CurveValues) && Context.Pin->Direction == EGPD_Input)
 			{
 				FString PinName = Context.Pin->PinFriendlyName.ToString();
-				FUIAction Action = FUIAction( FExecuteAction::CreateUObject(this, &UAnimGraphNode_ModifyCurve::RemoveCurvePin, FName(*PinName)) );
+				FUIAction Action = FUIAction( FExecuteAction::CreateUObject(const_cast<UAnimGraphNode_ModifyCurve*>(this), &UAnimGraphNode_ModifyCurve::RemoveCurvePin, FName(*PinName)) );
 				FText RemovePinLabelText = FText::Format(LOCTEXT("RemoveThisPin", "Remove This Curve Pin: {0}"), FText::FromString(PinName));
 				Context.MenuBuilder->AddMenuEntry(RemovePinLabelText, LOCTEXT("RemoveThisPinTooltip", "Remove this curve pin from this node"), FSlateIcon(), Action);
 			}
diff --git a/Engine/Source/Editor/AnimGraph/Private/K2Node_AnimGetter.cpp b/Engine/Source/Editor/AnimGraph/Private/K2Node_AnimGetter.cpp
--- a/Engine/Source/Editor/AnimGraph/Private/K2Node_AnimGetter.cpp
+++ b/Engine/Source/Editor/AnimGraph/Private/K2Node_AnimGetter.cpp
@@ -135,7 +135,7 @@ void UK2Node_AnimGetter::GetMenuActions(FBlueprintActionDatabaseRegistrar& Actio
 							Params.SourceNode = AssetNode;
 							Params.CachedTitle = Title;
 
-							UBlueprintNodeSpawner* Spawner = UBlueprintNodeSpawner::Create(UK2Node_AnimGetter::StaticClass(), /*AssetNode->GetGraph()*/nullptr, UBlueprintNodeSpawner::FCustomizeNodeDelegate::CreateUObject(this, &UK2Node_AnimGetter::PostSpawnNodeSetup, Params));
+							UBlueprintNodeSpawner* Spawner = UBlueprintNodeSpawner::Create(UK2Node_AnimGetter::StaticClass(), /*AssetNode->GetGraph()*/nullptr, UBlueprintNodeSpawner::FCustomizeNodeDelegate::CreateUObject(const_cast<UK2Node_AnimGetter*>(this), &UK2Node_AnimGetter::PostSpawnNodeSetup, Params));
 							Spawner->DynamicUiSignatureGetter = UBlueprintNodeSpawner::FUiSpecOverrideDelegate::CreateStatic(UiSpecOverride, Title);
 							ActionRegistrar.AddBlueprintAction(AnimBlueprint, Spawner);
 						}
@@ -161,7 +161,7 @@ void UK2Node_AnimGetter::GetMenuActions(FBlueprintActionDatabaseRegistrar& Actio
 							Params.SourceStateNode = StateNode;
 							Params.CachedTitle = Title;
 
-							UBlueprintNodeSpawner* Spawner = UBlueprintNodeSpawner::Create(UK2Node_AnimGetter::StaticClass(), /*StateNode->GetGraph()*/nullptr, UBlueprintNodeSpawner::FCustomizeNodeDelegate::CreateUObject(this, &UK2Node_AnimGetter::PostSpawnNodeSetup, Params));
+							UBlueprintNodeSpawner* Spawner = UBlueprintNodeSpawner::Create(UK2Node_AnimGetter::StaticClass(), /*StateNode->GetGraph()*/nullptr, UBlueprintNodeSpawner::FCustomizeNodeDelegate::CreateUObject(const_cast<UK2Node_AnimGetter*>(this), &UK2Node_AnimGetter::PostSpawnNodeSetup, Params));
 							Spawner->DynamicUiSignatureGetter = UBlueprintNodeSpawner::FUiSpecOverrideDelegate::CreateStatic(UiSpecOverride, Title);
 							ActionRegistrar.AddBlueprintAction(AnimBlueprint, Spawner);
 						}
@@ -183,7 +183,7 @@ void UK2Node_AnimGetter::GetMenuActions(FBlueprintActionDatabaseRegistrar& Actio
 							Params.SourceStateNode = TransitionNode;
 							Params.CachedTitle = Title;
 
-							UBlueprintNodeSpawner* Spawner = UBlueprintNodeSpawner::Create(UK2Node_AnimGetter::StaticClass(), /*TransitionNode->GetGraph()*/nullptr, UBlueprintNodeSpawner::FCustomizeNodeDelegate::CreateUObject(this, &UK2Node_AnimGetter::PostSpawnNodeSetup, Params));
+							UBlueprintNodeSpawner* Spawner = UBlueprintNodeSpawner::Create(UK2Node_AnimGetter::StaticClass(), /*TransitionNode->GetGraph()*/nullptr, UBlueprintNodeSpawner::FCustomizeNodeDelegate::CreateUObject(const_cast<UK2Node_AnimGetter*>(this), &UK2Node_AnimGetter::PostSpawnNodeSetup, Params));
 							Spawner->DynamicUiSignatureGetter = UBlueprintNodeSpawner::FUiSpecOverrideDelegate::CreateStatic(UiSpecOverride, Title);
 							ActionRegistrar.AddBlueprintAction(AnimBlueprint, Spawner);
 						}
@@ -197,7 +197,7 @@ void UK2Node_AnimGetter::GetMenuActions(FBlueprintActionDatabaseRegistrar& Actio
 							Params.SourceNode = MachineNode;
 							Params.CachedTitle = Title;
 
-							UBlueprintNodeSpawner* Spawner = UBlueprintNodeSpawner::Create(UK2Node_AnimGetter::StaticClass(), /*MachineNode*/nullptr, UBlueprintNodeSpawner::FCustomizeNodeDelegate::CreateUObject(this, &UK2Node_AnimGetter::PostSpawnNodeSetup, Params));
+							UBlueprintNodeSpawner* Spawner = UBlueprintNodeSpawner::Create(UK2Node_AnimGetter::StaticClass(), /*MachineNode*/nullptr, UBlueprintNodeSpawner::FCustomizeNodeDelegate::CreateUObject(const_cast<UK2Node_AnimGetter*>(this), &UK2Node_AnimGetter::PostSpawnNodeSetup, Params));
 							Spawner->DynamicUiSignatureGetter = UBlueprintNodeSpawner::FUiSpecOverrideDelegate::CreateStatic(UiSpecOverride, Title);
 							ActionRegistrar.AddBlueprintAction(AnimBlueprint, Spawner);
 						}
@@ -209,7 +209,7 @@ void UK2Node_AnimGetter::GetMenuActions(FBlueprintActionDatabaseRegistrar& Actio
 					FText Title = FText::Format(LOCTEXT("NodeTitleNoNode", "{0}"), Getter->GetDisplayNameText());
 					Params.CachedTitle = Title;
 
-					UBlueprintNodeSpawner* Spawner = UBlueprintNodeSpawner::Create(UK2Node_AnimGetter::StaticClass(), nullptr, UBlueprintNodeSpawner::FCustomizeNodeDelegate::CreateUObject(this, &UK2Node_AnimGetter::PostSpawnNodeSetup, Params));
+					UBlueprintNodeSpawner* Spawner = UBlueprintNodeSpawner::Create(UK2Node_AnimGetter::StaticClass(), nullptr, UBlueprintNodeSpawner::FCustomizeNodeDelegate::CreateUObject(const_cast<UK2Node_AnimGetter*>(this), &UK2Node_AnimGetter::PostSpawnNodeSetup, Params));
 					Spawner->DynamicUiSignatureGetter = UBlueprintNodeSpawner::FUiSpecOverrideDelegate::CreateStatic(UiSpecOverride, Title);
 					ActionRegistrar.AddBlueprintAction(AnimBlueprint, Spawner);
 				}
diff --git a/Engine/Source/Editor/BehaviorTreeEditor/Private/DetailCustomizations/BehaviorDecoratorDetails.cpp b/Engine/Source/Editor/BehaviorTreeEditor/Private/DetailCustomizations/BehaviorDecoratorDetails.cpp
--- a/Engine/Source/Editor/BehaviorTreeEditor/Private/DetailCustomizations/BehaviorDecoratorDetails.cpp
+++ b/Engine/Source/Editor/BehaviorTreeEditor/Private/DetailCustomizations/BehaviorDecoratorDetails.cpp
@@ -139,7 +139,7 @@ TSharedRef<SWidget> FBehaviorDecoratorDetails::OnGetAbortModeContent() const
 
 	for (int32 i = 0; i < ModeValues.Num(); i++)
 	{
-		FUIAction ItemAction( FExecuteAction::CreateSP( this, &FBehaviorDecoratorDetails::OnAbortModeChange, ModeValues[i].Int ) );
+		FUIAction ItemAction( FExecuteAction::CreateSP( const_cast<FBehaviorDecoratorDetails*>(this), &FBehaviorDecoratorDetails::OnAbortModeChange, ModeValues[i].Int ) );
 		MenuBuilder.AddMenuEntry(FText::FromString( ModeValues[i].Str ), TAttribute<FText>(), FSlateIcon(), ItemAction);
 	}
 
diff --git a/Engine/Source/Editor/BehaviorTreeEditor/Private/DetailCustomizations/BlackboardDecoratorDetails.cpp b/Engine/Source/Editor/BehaviorTreeEditor/Private/DetailCustomizations/BlackboardDecoratorDetails.cpp
--- a/Engine/Source/Editor/BehaviorTreeEditor/Private/DetailCustomizations/BlackboardDecoratorDetails.cpp
+++ b/Engine/Source/Editor/BehaviorTreeEditor/Private/DetailCustomizations/BlackboardDecoratorDetails.cpp
@@ -189,7 +189,7 @@ TSharedRef<SWidget> FBlackboardDecoratorDetails::OnGetEnumValueContent() const
 
 	for (int32 i = 0; i < EnumPropValues.Num(); i++)
 	{
-		FUIAction ItemAction( FExecuteAction::CreateSP( this, &FBlackboardDecoratorDetails::OnEnumValueComboChange, i ) );
+		FUIAction ItemAction( FExecuteAction::CreateSP( const_cast<FBlackboardDecoratorDetails*>(this), &FBlackboardDecoratorDetails::OnEnumValueComboChange, i ) );
 		MenuBuilder.AddMenuEntry( FText::FromString( EnumPropValues[i] ), TAttribute<FText>(), FSlateIcon(), ItemAction);
 	}
 
diff --git a/Engine/Source/Editor/BehaviorTreeEditor/Private/DetailCustomizations/BlackboardSelectorDetails.cpp b/Engine/Source/Editor/BehaviorTreeEditor/Private/DetailCustomizations/BlackboardSelectorDetails.cpp
--- a/Engine/Source/Editor/BehaviorTreeEditor/Private/DetailCustomizations/BlackboardSelectorDetails.cpp
+++ b/Engine/Source/Editor/BehaviorTreeEditor/Private/DetailCustomizations/BlackboardSelectorDetails.cpp
@@ -182,7 +182,7 @@ TSharedRef<SWidget> FBlackboardSelectorDetails::OnGetKeyContent() const
 
 	for (int32 Idx = 0; Idx < KeyValues.Num(); Idx++)
 	{
-		FUIAction ItemAction( FExecuteAction::CreateSP( this, &FBlackboardSelectorDetails::OnKeyComboChange, Idx) );
+		FUIAction ItemAction( FExecuteAction::CreateSP( const_cast<FBlackboardSelectorDetails*>(this), &FBlackboardSelectorDetails::OnKeyComboChange, Idx) );
 		MenuBuilder.AddMenuEntry( FText::FromName( KeyValues[Idx] ), TAttribute<FText>(), FSlateIcon(), ItemAction);
 	}
 
diff --git a/Engine/Source/Editor/BehaviorTreeEditor/Private/SBehaviorTreeBlackboardEditor.cpp b/Engine/Source/Editor/BehaviorTreeEditor/Private/SBehaviorTreeBlackboardEditor.cpp
--- a/Engine/Source/Editor/BehaviorTreeEditor/Private/SBehaviorTreeBlackboardEditor.cpp
+++ b/Engine/Source/Editor/BehaviorTreeEditor/Private/SBehaviorTreeBlackboardEditor.cpp
@@ -207,7 +207,7 @@ TSharedRef<SWidget> SBehaviorTreeBlackboardEditor::HandleCreateNewEntryMenu() co
 	Options.NameTypeToDisplay = EClassViewerNameTypeToDisplay::DisplayName;
 	Options.ClassFilter = MakeShareable( new FBlackboardEntryClassFilter );
 
-	FOnClassPicked OnPicked( FOnClassPicked::CreateRaw( this, &SBehaviorTreeBlackboardEditor::HandleKeyClassPicked ) );
+	FOnClassPicked OnPicked( FOnClassPicked::CreateRaw( const_cast<SBehaviorTreeBlackboardEditor*>(this), &SBehaviorTreeBlackboardEditor::HandleKeyClassPicked ) );
 
 	// clear the search box, just in case there's something typed in there 
 	// We need to do that since key adding code takes advantage of selection mechanics
diff --git a/Engine/Source/Editor/DetailCustomizations/Private/ComponentTransformDetails.cpp b/Engine/Source/Editor/DetailCustomizations/Private/ComponentTransformDetails.cpp
--- a/Engine/Source/Editor/DetailCustomizations/Private/ComponentTransformDetails.cpp
+++ b/Engine/Source/Editor/DetailCustomizations/Private/ComponentTransformDetails.cpp
@@ -285,15 +285,15 @@ FUIAction FComponentTransformDetails::CreateCopyAction( ETransformField::Type Tr
 	return
 		FUIAction
 		(
-			FExecuteAction::CreateSP(this, &FComponentTransformDetails::OnCopy, TransformField ),
-			FCanExecuteAction::CreateSP(this, &FComponentTransformDetails::OnCanCopy, TransformField )
+			FExecuteAction::CreateSP(const_cast<FComponentTransformDetails*>(this), &FComponentTransformDetails::OnCopy, TransformField ),
+			FCanExecuteAction::CreateSP(const_cast<FComponentTransformDetails*>(this), &FComponentTransformDetails::OnCanCopy, TransformField )
 		);
 }
 
 FUIAction FComponentTransformDetails::CreatePasteAction( ETransformField::Type TransformField ) const
 {
 	return 
-		 FUIAction( FExecuteAction::CreateSP(this, &FComponentTransformDetails::OnPaste, TransformField ) );
+		 FUIAction( FExecuteAction::CreateSP(const_cast<FComponentTransformDetails*>(this), &FComponentTransformDetails::OnPaste, TransformField ) );
 }
 
 BEGIN_SLATE_FUNCTION_BUILD_OPTIMIZATION
diff --git a/Engine/Source/Editor/DetailCustomizations/Private/DeviceProfileDetails.cpp b/Engine/Source/Editor/DetailCustomizations/Private/DeviceProfileDetails.cpp
--- a/Engine/Source/Editor/DetailCustomizations/Private/DeviceProfileDetails.cpp
+++ b/Engine/Source/Editor/DetailCustomizations/Private/DeviceProfileDetails.cpp
@@ -747,7 +747,7 @@ void FDeviceProfileConsoleVariablesPropertyDetails::CreateRowWidgetForCVarProper
 			SNew(SEditableTextBox)
 			.Text(FText::FromString(CVarValueAsString))
 			.SelectAllTextWhenFocused(true)
-			.OnTextCommitted(this, &FDeviceProfileConsoleVariablesPropertyDetails::OnCVarValueCommited, InProperty)
+			.OnTextCommitted(const_cast<FDeviceProfileConsoleVariablesPropertyDetails*>(this), &FDeviceProfileConsoleVariablesPropertyDetails::OnCVarValueCommited, InProperty)
 		]
 		+ SHorizontalBox::Slot()
 		.Padding(DeviceProfilePropertyConstants::PropertyPadding)
@@ -755,7 +755,7 @@ void FDeviceProfileConsoleVariablesPropertyDetails::CreateRowWidgetForCVarProper
 		[
 			SNew(SButton)
 			.ButtonStyle(FEditorStyle::Get(), "HoverHintOnly")
-			.OnClicked(this, &FDeviceProfileConsoleVariablesPropertyDetails::OnRemoveCVarProperty, InProperty)
+			.OnClicked(const_cast<FDeviceProfileConsoleVariablesPropertyDetails*>(this), &FDeviceProfileConsoleVariablesPropertyDetails::OnRemoveCVarProperty, InProperty)
 			.ContentPadding(4.0f)
 			.ForegroundColor(FSlateColor::UseForeground())
 			.IsFocusable(false)
diff --git a/Engine/Source/Editor/DetailCustomizations/Private/MarginCustomization.cpp b/Engine/Source/Editor/DetailCustomizations/Private/MarginCustomization.cpp
--- a/Engine/Source/Editor/DetailCustomizations/Private/MarginCustomization.cpp
+++ b/Engine/Source/Editor/DetailCustomizations/Private/MarginCustomization.cpp
@@ -92,10 +92,10 @@ TSharedRef<SWidget> FMarginStructCustomization::MakeChildPropertyWidget( int32 P
 		.Value( this, &FMarginStructCustomization::OnGetValue, PropertyIndex )
 		.Font( IDetailLayoutBuilder::GetDetailFont() )
 		.UndeterminedString( NSLOCTEXT( "PropertyEditor", "MultipleValues", "Multiple Values") )
-		.OnValueCommitted( this, &FMarginStructCustomization::OnValueCommitted, PropertyIndex )
-		.OnValueChanged( this, &FMarginStructCustomization::OnValueChanged, PropertyIndex )
-		.OnBeginSliderMovement( this, &FMarginStructCustomization::OnBeginSliderMovement )
-		.OnEndSliderMovement( this, &FMarginStructCustomization::OnEndSliderMovement )
+		.OnValueCommitted( const_cast<FMarginStructCustomization*>(this), &FMarginStructCustomization::OnValueCommitted, PropertyIndex )
+		.OnValueChanged( const_cast<FMarginStructCustomization*>(this), &FMarginStructCustomization::OnValueChanged, PropertyIndex )
+		.OnBeginSliderMovement( const_cast<FMarginStructCustomization*>(this), &FMarginStructCustomization::OnBeginSliderMovement )
+		.OnEndSliderMovement( const_cast<FMarginStructCustomization*>(this), &FMarginStructCustomization::OnEndSliderMovement )
 		.LabelVAlign( VAlign_Center )
 		.AllowSpin( bIsMarginUsingUVSpace ? true : false )
 		.MinValue( bIsMarginUsingUVSpace ? 0.0f : TNumericLimits<float>::Lowest() )
diff --git a/Engine/Source/Editor/GameProjectGeneration/Private/SGetSuggestedIDEWidget.cpp b/Engine/Source/Editor/GameProjectGeneration/Private/SGetSuggestedIDEWidget.cpp
--- a/Engine/Source/Editor/GameProjectGeneration/Private/SGetSuggestedIDEWidget.cpp
+++ b/Engine/Source/Editor/GameProjectGeneration/Private/SGetSuggestedIDEWidget.cpp
@@ -33,14 +33,14 @@ TSharedRef<SWidget> SGetSuggestedIDEWidget::CreateGetSuggestedIDEWidget() const
 		// If the installer for this platform's IDE can be downloaded and launched directly, show a button
 		return SNew(SButton)
 			.Text(FText::Format(LOCTEXT("IDEInstallButtonText", "Install {0}"), FSourceCodeNavigation::GetSuggestedSourceCodeIDE()))
-			.OnClicked(this, &SGetSuggestedIDEWidget::OnInstallIDEClicked);
+			.OnClicked(const_cast<SGetSuggestedIDEWidget*>(this), &SGetSuggestedIDEWidget::OnInstallIDEClicked);
 	}
 	else
 	{
 		// If the user must open a web page, show a link
 		return SNew(SHyperlink)
 			.Text(FText::Format(LOCTEXT("IDEDownloadLinkText", "Download {0}"), FSourceCodeNavigation::GetSuggestedSourceCodeIDE()))
-			.OnNavigate(this, &SGetSuggestedIDEWidget::OnDownloadIDEClicked, FSourceCodeNavigation::GetSuggestedSourceCodeIDEDownloadURL());
+			.OnNavigate(const_cast<SGetSuggestedIDEWidget*>(this), &SGetSuggestedIDEWidget::OnDownloadIDEClicked, FSourceCodeNavigation::GetSuggestedSourceCodeIDEDownloadURL());
 	}
 }
 
diff --git a/Engine/Source/Editor/GameProjectGeneration/Private/SProjectBrowser.cpp b/Engine/Source/Editor/GameProjectGeneration/Private/SProjectBrowser.cpp
--- a/Engine/Source/Editor/GameProjectGeneration/Private/SProjectBrowser.cpp
+++ b/Engine/Source/Editor/GameProjectGeneration/Private/SProjectBrowser.cpp
@@ -396,10 +396,10 @@ void SProjectBrowser::ConstructCategory( const TSharedRef<SVerticalBox>& InCateg
 		.SelectionMode(ESelectionMode::Single)
 		.ClearSelectionOnClick(false)
 		.AllowOverscroll(EAllowOverscroll::No)
-		.OnGenerateTile(this, &SProjectBrowser::MakeProjectViewWidget)
+		.OnGenerateTile(const_cast<SProjectBrowser*>(this), &SProjectBrowser::MakeProjectViewWidget)
 		.OnContextMenuOpening(this, &SProjectBrowser::OnGetContextMenuContent)
-		.OnMouseButtonDoubleClick(this, &SProjectBrowser::HandleProjectItemDoubleClick)
-		.OnSelectionChanged(this, &SProjectBrowser::HandleProjectViewSelectionChanged, Category->CategoryName)
+		.OnMouseButtonDoubleClick(const_cast<SProjectBrowser*>(this), &SProjectBrowser::HandleProjectItemDoubleClick)
+		.OnSelectionChanged(const_cast<SProjectBrowser*>(this), &SProjectBrowser::HandleProjectViewSelectionChanged, Category->CategoryName)
 		.ItemHeight(ThumbnailSize + ThumbnailBorderPadding + 32)
 		.ItemWidth(ThumbnailSize + ThumbnailBorderPadding)
 	];
diff --git a/Engine/Source/Editor/HierarchicalLODOutliner/Private/HLODOutliner.cpp b/Engine/Source/Editor/HierarchicalLODOutliner/Private/HLODOutliner.cpp
--- a/Engine/Source/Editor/HierarchicalLODOutliner/Private/HLODOutliner.cpp
+++ b/Engine/Source/Editor/HierarchicalLODOutliner/Private/HLODOutliner.cpp
@@ -981,7 +981,7 @@ namespace HLODOutliner
 			LOCTEXT("AutoLODTooltip", "Determine LOD level automatically"),
 			FSlateIcon(),
 			FUIAction(
-				FExecuteAction::CreateSP(this, &SHLODOutliner::SetForcedLODLevel, -1), 
+				FExecuteAction::CreateSP(const_cast<SHLODOutliner*>(this), &SHLODOutliner::SetForcedLODLevel, -1), 
 				FCanExecuteAction(), 
 				FGetActionCheckState::CreateLambda([this](){ return ForcedLODLevel == -1 ? ECheckBoxState::Checked : ECheckBoxState::Unchecked; })),
 			NAME_None,
@@ -999,7 +999,7 @@ namespace HLODOutliner
 						FText::Format(LOCTEXT("LODLevelTooltipFormat", "Force LOD to level {0}"), FText::AsNumber(LODIndex)),
 						FSlateIcon(),
 						FUIAction(
-							FExecuteAction::CreateSP(this, &SHLODOutliner::SetForcedLODLevel, LODIndex), 
+							FExecuteAction::CreateSP(const_cast<SHLODOutliner*>(this), &SHLODOutliner::SetForcedLODLevel, LODIndex), 
 							FCanExecuteAction(), 
 							FGetActionCheckState::CreateLambda([this, LODIndex](){ return ForcedLODLevel == LODIndex ? ECheckBoxState::Checked : ECheckBoxState::Unchecked; })),
 						NAME_None,
diff --git a/Engine/Source/Editor/IntroTutorials/Private/IntroTutorials.cpp b/Engine/Source/Editor/IntroTutorials/Private/IntroTutorials.cpp
--- a/Engine/Source/Editor/IntroTutorials/Private/IntroTutorials.cpp
+++ b/Engine/Source/Editor/IntroTutorials/Private/IntroTutorials.cpp
@@ -78,7 +78,7 @@ TSharedRef<FExtender> FIntroTutorials::AddSummonBlueprintTutorialsMenuExtender(c
 		"HelpBrowse",
 		EExtensionHook::After,
 		CommandList,
-		FMenuExtensionDelegate::CreateRaw(this, &FIntroTutorials::AddSummonBlueprintTutorialsMenuExtension, PrimaryObject));
+		FMenuExtensionDelegate::CreateRaw(const_cast<FIntroTutorials*>(this), &FIntroTutorials::AddSummonBlueprintTutorialsMenuExtension, PrimaryObject));
 
 	return Extender;
 }
diff --git a/Engine/Source/Editor/Kismet/Private/BlueprintEditorTabFactories.cpp b/Engine/Source/Editor/Kismet/Private/BlueprintEditorTabFactories.cpp
--- a/Engine/Source/Editor/Kismet/Private/BlueprintEditorTabFactories.cpp
+++ b/Engine/Source/Editor/Kismet/Private/BlueprintEditorTabFactories.cpp
@@ -197,7 +197,7 @@ TSharedRef<SWidget> FDefaultsEditorSummoner::CreateOptionalDataOnlyMessage() con
 					[
 						SNew(SHyperlink)
 						.Style(FEditorStyle::Get(), "Common.GotoBlueprintHyperlink")
-						.OnNavigate(this, &FDefaultsEditorSummoner::OnChangeBlueprintToNotDataOnly)
+						.OnNavigate(const_cast<FDefaultsEditorSummoner*>(this), &FDefaultsEditorSummoner::OnChangeBlueprintToNotDataOnly)
 						.Text(LOCTEXT("FullEditor", "Open Full Blueprint Editor"))
 						.ToolTipText(LOCTEXT("FullEditorToolTip", "This opens the blueprint in the full editor."))
 					]
diff --git a/Engine/Source/Editor/Kismet/Private/SBlueprintSubPalette.cpp b/Engine/Source/Editor/Kismet/Private/SBlueprintSubPalette.cpp
--- a/Engine/Source/Editor/Kismet/Private/SBlueprintSubPalette.cpp
+++ b/Engine/Source/Editor/Kismet/Private/SBlueprintSubPalette.cpp
@@ -376,7 +376,7 @@ void SBlueprintSubPalette::BindCommands(TSharedPtr<FUICommandList> CommandListIn
 
 	CommandListIn->MapAction(
 		PaletteCommands.RefreshPalette,
-		FExecuteAction::CreateSP(this, &SBlueprintSubPalette::RefreshActionsList, /*bPreserveExpansion =*/true)
+		FExecuteAction::CreateSP(const_cast<SBlueprintSubPalette*>(this), &SBlueprintSubPalette::RefreshActionsList, /*bPreserveExpansion =*/true)
 	);
 }
 
diff --git a/Engine/Source/Editor/Kismet/Private/STimelineEditor.cpp b/Engine/Source/Editor/Kismet/Private/STimelineEditor.cpp
--- a/Engine/Source/Editor/Kismet/Private/STimelineEditor.cpp
+++ b/Engine/Source/Editor/Kismet/Private/STimelineEditor.cpp
@@ -1657,7 +1657,7 @@ TSharedPtr< SWidget > STimelineEditor::MakeContextMenu() const
 	{
 		TSharedRef<SWidget> SizeSlider = SNew(SSlider)
 			.Value(this, &STimelineEditor::GetSizeScaleValue)
-			.OnValueChanged(this, &STimelineEditor::SetSizeScaleValue);
+			.OnValueChanged(const_cast<STimelineEditor*>(this), &STimelineEditor::SetSizeScaleValue);
 
 		MenuBuilder.AddWidget(SizeSlider, LOCTEXT("TimelineEditorVerticalSize", "Height"));
 	}
diff --git a/Engine/Source/Editor/LevelEditor/Private/SLevelViewportToolBar.cpp b/Engine/Source/Editor/LevelEditor/Private/SLevelViewportToolBar.cpp
--- a/Engine/Source/Editor/LevelEditor/Private/SLevelViewportToolBar.cpp
+++ b/Engine/Source/Editor/LevelEditor/Private/SLevelViewportToolBar.cpp
@@ -628,7 +628,7 @@ TSharedRef<SWidget> SLevelViewportToolBar::GenerateDevicePreviewMenu() const
 
 	// Default menu - clear all settings
 	{
-		FUIAction Action( FExecuteAction::CreateSP( this, &SLevelViewportToolBar::SetLevelProfile, FString( "Default" ) ),
+		FUIAction Action( FExecuteAction::CreateSP( const_cast<SLevelViewportToolBar*>(this), &SLevelViewportToolBar::SetLevelProfile, FString( "Default" ) ),
 			FCanExecuteAction(),
 			FIsActionChecked::CreateSP( ViewportRef, &SLevelViewport::IsDeviceProfileStringSet, FString( "Default" ) ) );
 		DeviceMenuBuilder.AddMenuEntry( LOCTEXT("DevicePreviewMenuClear", "Off"), FText::GetEmpty(), FSlateIcon(), Action, NAME_None, EUserInterfaceActionType::Button );
@@ -651,7 +651,7 @@ TSharedRef<SWidget> SLevelViewportToolBar::GenerateDevicePreviewMenu() const
 		{
 			const FName PlatformIcon = UIManager->GetDeviceIconName( CurItem );
 
-			FUIAction Action( FExecuteAction::CreateSP( this, &SLevelViewportToolBar::SetLevelProfile, CurItem ),
+			FUIAction Action( FExecuteAction::CreateSP( const_cast<SLevelViewportToolBar*>(this), &SLevelViewportToolBar::SetLevelProfile, CurItem ),
 				FCanExecuteAction(),
 				FIsActionChecked::CreateSP( ViewportRef, &SLevelViewport::IsDeviceProfileStringSet, CurItem ) );
 			DeviceMenuBuilder.AddMenuEntry( FText::FromString( CurItem ), FText(), FSlateIcon(FEditorStyle::GetStyleSetName(), PlatformIcon), Action, NAME_None, EUserInterfaceActionType::Button );
@@ -675,7 +675,7 @@ TSharedRef<SWidget> SLevelViewportToolBar::GenerateDevicePreviewMenu() const
 			DeviceMenuBuilder.AddSubMenu(
 				FText::FromString( PlatformNameStr ),
 				FText::GetEmpty(),
-				FNewMenuDelegate::CreateRaw( this, &SLevelViewportToolBar::MakeDevicePreviewSubMenu, DeviceProfiles ),
+				FNewMenuDelegate::CreateRaw( const_cast<SLevelViewportToolBar*>(this), &SLevelViewportToolBar::MakeDevicePreviewSubMenu, DeviceProfiles ),
 				false,
 				FSlateIcon(FEditorStyle::GetStyleSetName(), PlatformIcon)
 				);
@@ -1062,7 +1062,7 @@ TSharedRef<SWidget> SLevelViewportToolBar::GenerateFOVMenu() const
 				.MinValue(FOVMin)
 				.MaxValue(FOVMax)
 				.Value( this, &SLevelViewportToolBar::OnGetFOVValue )
-				.OnValueChanged( this, &SLevelViewportToolBar::OnFOVValueChanged )
+				.OnValueChanged( const_cast<SLevelViewportToolBar*>(this), &SLevelViewportToolBar::OnFOVValueChanged )
 			]
 		];
 }
@@ -1114,7 +1114,7 @@ TSharedRef<SWidget> SLevelViewportToolBar::GenerateScreenPercentageMenu() const
 		.MinValue(PreviewScreenPercentageMin)
 		.MaxValue(PreviewScreenPercentageMax)
 		.Value(this, &SLevelViewportToolBar::OnGetScreenPercentageValue)
-		.OnValueChanged(this, &SLevelViewportToolBar::OnScreenPercentageValueChanged)
+		.OnValueChanged(const_cast<SLevelViewportToolBar*>(this), &SLevelViewportToolBar::OnScreenPercentageValueChanged)
 		]
 		];
 }
diff --git a/Engine/Source/Editor/Matinee/Private/MatineeTrackHelpers.cpp b/Engine/Source/Editor/Matinee/Private/MatineeTrackHelpers.cpp
--- a/Engine/Source/Editor/Matinee/Private/MatineeTrackHelpers.cpp
+++ b/Engine/Source/Editor/Matinee/Private/MatineeTrackHelpers.cpp
@@ -185,7 +185,7 @@ bool UMatineeTrackAnimControlHelper::PreCreateTrack( UInterpGroup* Group, const
 				SNew(STextComboPopup)
 				.Label(NSLOCTEXT("Matinee.Popups", "ChooseAnimSlot", "Choose Anim Slot..."))
 				.TextOptions(SlotStrings)
-				.OnTextChosen_UObject(this, &UMatineeTrackAnimControlHelper::OnCreateTrackTextEntry, NewWindow, (FString *)&Result)
+				.OnTextChosen_UObject(const_cast<UMatineeTrackAnimControlHelper*>(this), &UMatineeTrackAnimControlHelper::OnCreateTrackTextEntry, NewWindow, (FString *)&Result)
 				;
 
 			NewWindow->SetContent(TextEntryPopup);
@@ -282,7 +282,7 @@ bool UMatineeTrackAnimControlHelper::PreCreateKeyframe( UInterpTrack *Track, flo
 		if ( Parent.IsValid() )
 		{
 			FAssetPickerConfig AssetPickerConfig;
-			AssetPickerConfig.OnAssetSelected = FOnAssetSelected::CreateUObject( this, &UMatineeTrackAnimControlHelper::OnAddKeyTextEntry, Mode->InterpEd, Track );
+			AssetPickerConfig.OnAssetSelected = FOnAssetSelected::CreateUObject( const_cast<UMatineeTrackAnimControlHelper*>(this), &UMatineeTrackAnimControlHelper::OnAddKeyTextEntry, Mode->InterpEd, Track );
 			AssetPickerConfig.bAllowNullSelection = false;
 			AssetPickerConfig.InitialAssetViewType = EAssetViewType::List;
 
@@ -377,7 +377,7 @@ bool UMatineeTrackDirectorHelper::PreCreateKeyframe( UInterpTrack *Track, float
 			SNew(STextComboPopup)
 			.Label(NSLOCTEXT("Matinee.Popups", "NewCut", "Cut to Group..."))
 			.TextOptions(GroupNames)
-			.OnTextChosen_UObject(this, &UMatineeTrackDirectorHelper::OnAddKeyTextEntry, Mode->InterpEd, Track);
+			.OnTextChosen_UObject(const_cast<UMatineeTrackDirectorHelper*>(this), &UMatineeTrackDirectorHelper::OnAddKeyTextEntry, Mode->InterpEd, Track);
 
 		TSharedPtr< SWindow > Parent = FSlateApplication::Get().GetActiveTopLevelWindow();
 		if ( Parent.IsValid() )
@@ -435,7 +435,7 @@ bool UMatineeTrackEventHelper::PreCreateKeyframe( UInterpTrack *Track, float Key
 		SNew(STextEntryPopup)
 		.Label(NSLOCTEXT("Matinee.Popups", "NewEventName", "New Event Name"))
 		.DefaultText(FText::FromString(TEXT("Event")))
-		.OnTextCommitted_UObject(this, &UMatineeTrackEventHelper::OnAddKeyTextEntry, (IMatineeBase*)Mode->InterpEd, Track)
+		.OnTextCommitted_UObject(const_cast<UMatineeTrackEventHelper*>(this), &UMatineeTrackEventHelper::OnAddKeyTextEntry, (IMatineeBase*)Mode->InterpEd, Track)
 		.SelectAllTextWhenFocused(true)
 		.ClearKeyboardFocusOnCommit(false)
 		.MaxWidth(1024.0f)
@@ -564,7 +564,7 @@ bool UMatineeTrackFloatPropHelper::PreCreateTrack( UInterpGroup* Group, const UI
 					SNew(STextComboPopup)
 					.Label(NSLOCTEXT("Matinee.Popups", "PropertyName", "Property Name"))
 					.TextOptions(PropStrings)
-					.OnTextChosen_UObject(this, &UMatineeTrackFloatPropHelper::OnCreateTrackTextEntry, NewWindow, (FString *)&Result)
+					.OnTextChosen_UObject(const_cast<UMatineeTrackFloatPropHelper*>(this), &UMatineeTrackFloatPropHelper::OnCreateTrackTextEntry, NewWindow, (FString *)&Result)
 					;
 
 				NewWindow->SetContent(TextEntryPopup);
@@ -702,7 +702,7 @@ bool UMatineeTrackBoolPropHelper::PreCreateTrack( UInterpGroup* Group, const UIn
 					SNew(STextComboPopup)
 					.Label(NSLOCTEXT("Matinee.Popups", "PropertyName", "Property Name"))
 					.TextOptions(PropStrings)
-					.OnTextChosen_UObject(this, &UMatineeTrackBoolPropHelper::OnCreateTrackTextEntry, TWeakPtr<SWindow>(NewWindow), (FString *)&Result);
+					.OnTextChosen_UObject(const_cast<UMatineeTrackBoolPropHelper*>(this), &UMatineeTrackBoolPropHelper::OnCreateTrackTextEntry, TWeakPtr<SWindow>(NewWindow), (FString *)&Result);
 
 				NewWindow->SetContent(TextEntryPopup);
 				GEditor->EditorAddModalWindow(NewWindow);
@@ -818,7 +818,7 @@ bool UMatineeTrackToggleHelper::PreCreateKeyframe( UInterpTrack *Track, float Ke
 		SNew(STextComboPopup)
 		.Label(NSLOCTEXT("Matinee.Popups", "ToggleAction", "Toggle Action"))
 		.TextOptions(PropStrings)
-		.OnTextChosen_UObject(this, &UMatineeTrackToggleHelper::OnAddKeyTextEntry, InterpEd, Track)
+		.OnTextChosen_UObject(const_cast<UMatineeTrackToggleHelper*>(this), &UMatineeTrackToggleHelper::OnAddKeyTextEntry, InterpEd, Track)
 		;
 
 	TSharedPtr< SWindow > Parent = FSlateApplication::Get().GetActiveTopLevelWindow();
@@ -905,7 +905,7 @@ bool UMatineeTrackVectorPropHelper::ChooseProperty(TArray<FName> &PropNames) con
 		SNew(STextComboPopup)
 		.Label(NSLOCTEXT("Matinee.Popups", "PropertyName", "Property Name"))
 		.TextOptions(PropStrings)
-		.OnTextChosen_UObject(this, &UMatineeTrackVectorPropHelper::OnCreateTrackTextEntry, TWeakPtr<SWindow>(NewWindow), (FString *)&Result)
+		.OnTextChosen_UObject(const_cast<UMatineeTrackVectorPropHelper*>(this), &UMatineeTrackVectorPropHelper::OnCreateTrackTextEntry, TWeakPtr<SWindow>(NewWindow), (FString *)&Result)
 		;
 
 	NewWindow->SetContent(TextEntryPopup);
@@ -1170,7 +1170,7 @@ bool UMatineeTrackVisibilityHelper::PreCreateKeyframe( UInterpTrack *Track, floa
 		SNew(STextComboPopup)
 		.Label(NSLOCTEXT("Matinee.Popups", "VisibilityAction", "Visibility Action"))
 		.TextOptions(PropStrings)
-		.OnTextChosen_UObject(this, &UMatineeTrackVisibilityHelper::OnAddKeyTextEntry, Mode->InterpEd, Track)
+		.OnTextChosen_UObject(const_cast<UMatineeTrackVisibilityHelper*>(this), &UMatineeTrackVisibilityHelper::OnAddKeyTextEntry, Mode->InterpEd, Track)
 		;
 
 	TSharedPtr< SWindow > Parent = FSlateApplication::Get().GetActiveTopLevelWindow();
diff --git a/Engine/Source/Editor/Persona/Private/SAnimTrackCurvePanel.cpp b/Engine/Source/Editor/Persona/Private/SAnimTrackCurvePanel.cpp
--- a/Engine/Source/Editor/Persona/Private/SAnimTrackCurvePanel.cpp
+++ b/Engine/Source/Editor/Persona/Private/SAnimTrackCurvePanel.cpp
@@ -894,7 +894,7 @@ TSharedRef<SWidget> SAnimTrackCurvePanel::CreateCurveContextMenu(USkeleton::Anim
 			MenuBuilder.AddWidget(
 				SNew(SCheckBox)
 				.IsChecked(this, &SAnimTrackCurvePanel::GetCurveFlagAsCheckboxState, CurveUid, AACF_Disabled)
-				.OnCheckStateChanged(this, &SAnimTrackCurvePanel::SetCurveFlagFromCheckboxState, CurveUid, AACF_Disabled)
+				.OnCheckStateChanged(const_cast<SAnimTrackCurvePanel*>(this), &SAnimTrackCurvePanel::SetCurveFlagFromCheckboxState, CurveUid, AACF_Disabled)
 				.ToolTipText(LOCTEXT("DisableCurveTooltip", "Disable Track"))
 				[
 					SNew(STextBlock)
diff --git a/Engine/Source/Editor/Persona/Private/SAnimViewportToolBar.cpp b/Engine/Source/Editor/Persona/Private/SAnimViewportToolBar.cpp
--- a/Engine/Source/Editor/Persona/Private/SAnimViewportToolBar.cpp
+++ b/Engine/Source/Editor/Persona/Private/SAnimViewportToolBar.cpp
@@ -359,7 +359,7 @@ TSharedRef<SWidget> SAnimViewportToolBar::MakeFloorOffsetWidget() const
 				.MinSliderValue(-100.0f)
 				.MaxSliderValue(100.0f)
 				.Value(this, &SAnimViewportToolBar::OnGetFloorOffset)
-				.OnValueChanged(this, &SAnimViewportToolBar::OnFloorOffsetChanged)
+				.OnValueChanged(const_cast<SAnimViewportToolBar*>(this), &SAnimViewportToolBar::OnFloorOffsetChanged)
 				.ToolTipText(LOCTEXT("FloorOffsetToolTip", "Height offset for the floor mesh (stored per-mesh)"))
 			]
 		];
@@ -386,8 +386,8 @@ TSharedRef<SWidget> SAnimViewportToolBar::MakeFOVWidget() const
 				.MinSliderValue(FOVMin)
 				.MaxSliderValue(FOVMax)
 				.Value(this, &SAnimViewportToolBar::OnGetFOVValue)
-				.OnValueChanged(this, &SAnimViewportToolBar::OnFOVValueChanged)
-				.OnValueCommitted(this, &SAnimViewportToolBar::OnFOVValueCommitted)
+				.OnValueChanged(const_cast<SAnimViewportToolBar*>(this), &SAnimViewportToolBar::OnFOVValueChanged)
+				.OnValueCommitted(const_cast<SAnimViewportToolBar*>(this), &SAnimViewportToolBar::OnFOVValueCommitted)
 			]
 		];
 }
diff --git a/Engine/Source/Editor/Persona/Private/SAnimationSequenceBrowser.cpp b/Engine/Source/Editor/Persona/Private/SAnimationSequenceBrowser.cpp
--- a/Engine/Source/Editor/Persona/Private/SAnimationSequenceBrowser.cpp
+++ b/Engine/Source/Editor/Persona/Private/SAnimationSequenceBrowser.cpp
@@ -811,7 +811,7 @@ TSharedRef<SWidget> SAnimationSequenceBrowser::CreateHistoryMenu(bool bInBackHis
 
 				MenuBuilder.AddMenuEntry(DisplayName, Tooltip, FSlateIcon(), 
 					FUIAction(
-					FExecuteAction::CreateRaw(this, &SAnimationSequenceBrowser::GoToHistoryIndex, HistoryIdx)
+					FExecuteAction::CreateRaw(const_cast<SAnimationSequenceBrowser*>(this), &SAnimationSequenceBrowser::GoToHistoryIndex, HistoryIdx)
 					), 
 					NAME_None, EUserInterfaceActionType::Button);
 			}
diff --git a/Engine/Source/Editor/Persona/Private/SMorphTargetViewer.cpp b/Engine/Source/Editor/Persona/Private/SMorphTargetViewer.cpp
--- a/Engine/Source/Editor/Persona/Private/SMorphTargetViewer.cpp
+++ b/Engine/Source/Editor/Persona/Private/SMorphTargetViewer.cpp
@@ -431,7 +431,7 @@ TSharedPtr<SWidget> SMorphTargetViewer::OnGetContextMenuContent() const
 		FUIAction Action;
 
 		{
-			Action.ExecuteAction = FExecuteAction::CreateSP(this, &SMorphTargetViewer::OnDeleteMorphTargets);
+			Action.ExecuteAction = FExecuteAction::CreateSP(const_cast<SMorphTargetViewer*>(this), &SMorphTargetViewer::OnDeleteMorphTargets);
 			Action.CanExecuteAction = FCanExecuteAction::CreateSP(this, &SMorphTargetViewer::CanPerformDelete);
 			const FText Label = LOCTEXT("DeleteMorphTargetButtonLabel", "Delete");
 			const FText ToolTipText = LOCTEXT("DeleteMorphTargetButtonTooltip", "Deletes the selected morph targets.");
@@ -439,7 +439,7 @@ TSharedPtr<SWidget> SMorphTargetViewer::OnGetContextMenuContent() const
 		}
 
 		{
-			Action.ExecuteAction = FExecuteAction::CreateSP(this, &SMorphTargetViewer::OnCopyMorphTargetNames);
+			Action.ExecuteAction = FExecuteAction::CreateSP(const_cast<SMorphTargetViewer*>(this), &SMorphTargetViewer::OnCopyMorphTargetNames);
 			Action.CanExecuteAction = nullptr;
 			const FText Label = LOCTEXT("CopyMorphTargetNamesButtonLabel", "Copy Names");
 			const FText ToolTipText = LOCTEXT("CopyMorphTargetNamesButtonTooltip", "Copy the names of selected morph targets to clipboard");
diff --git a/Engine/Source/Editor/Persona/Private/SPoseEditor.cpp b/Engine/Source/Editor/Persona/Private/SPoseEditor.cpp
--- a/Engine/Source/Editor/Persona/Private/SPoseEditor.cpp
+++ b/Engine/Source/Editor/Persona/Private/SPoseEditor.cpp
@@ -640,7 +640,7 @@ TSharedPtr<SWidget> SPoseViewer::OnGetContextMenuContentForCurveList() const
 
 	MenuBuilder.BeginSection("CurveAction", LOCTEXT("CurveActions", "Selected Item Actions"));
 	{
-		FUIAction Action = FUIAction(FExecuteAction::CreateSP(this, &SPoseViewer::OnDeleteCurves),
+		FUIAction Action = FUIAction(FExecuteAction::CreateSP(const_cast<SPoseViewer*>(this), &SPoseViewer::OnDeleteCurves),
 			FCanExecuteAction::CreateSP(this, &SPoseViewer::IsCurveSelected));
 		const FText MenuLabel = LOCTEXT("DeleteCurveButtonLabel", "Delete");
 		const FText MenuToolTip = LOCTEXT("DeleteCurveButtonTooltip", "Deletes the selected animation curve.");
diff --git a/Engine/Source/Editor/Persona/Private/SRetargetSourceWindow.cpp b/Engine/Source/Editor/Persona/Private/SRetargetSourceWindow.cpp
--- a/Engine/Source/Editor/Persona/Private/SRetargetSourceWindow.cpp
+++ b/Engine/Source/Editor/Persona/Private/SRetargetSourceWindow.cpp
@@ -337,7 +337,7 @@ TSharedPtr<SWidget> SRetargetSourceWindow::OnGetContextMenuContent() const
 
 	MenuBuilder.BeginSection("RetargetSourceAction", LOCTEXT( "New", "New" ) );
 	{
-		FUIAction Action = FUIAction( FExecuteAction::CreateSP( this, &SRetargetSourceWindow::OnAddRetargetSource ) );
+		FUIAction Action = FUIAction( FExecuteAction::CreateSP( const_cast<SRetargetSourceWindow*>(this), &SRetargetSourceWindow::OnAddRetargetSource ) );
 		const FText Label = LOCTEXT("AddRetargetSourceActionLabel", "Add...");
 		const FText ToolTipText = LOCTEXT("AddRetargetSourceActionTooltip", "Add new retarget source.");
 		MenuBuilder.AddMenuEntry( Label, ToolTipText, FSlateIcon(), Action);
@@ -346,21 +346,21 @@ TSharedPtr<SWidget> SRetargetSourceWindow::OnGetContextMenuContent() const
 
 	MenuBuilder.BeginSection("RetargetSourceAction", LOCTEXT( "Selected", "Selected Item Actions" ) );
 	{
-		FUIAction Action = FUIAction( FExecuteAction::CreateSP( this, &SRetargetSourceWindow::OnRenameRetargetSource ), 
+		FUIAction Action = FUIAction( FExecuteAction::CreateSP( const_cast<SRetargetSourceWindow*>(this), &SRetargetSourceWindow::OnRenameRetargetSource ), 
 			FCanExecuteAction::CreateSP( this, &SRetargetSourceWindow::CanPerformRename ) );
 		const FText Label = LOCTEXT("RenameRetargetSourceActionLabel", "Rename");
 		const FText ToolTipText = LOCTEXT("RenameRetargetSourceActionTooltip", "Rename the selected retarget source.");
 		MenuBuilder.AddMenuEntry( Label, ToolTipText, FSlateIcon(), Action);
 	}
 	{
-		FUIAction Action = FUIAction( FExecuteAction::CreateSP( this, &SRetargetSourceWindow::OnDeleteRetargetSource ), 
+		FUIAction Action = FUIAction( FExecuteAction::CreateSP( const_cast<SRetargetSourceWindow*>(this), &SRetargetSourceWindow::OnDeleteRetargetSource ), 
 			FCanExecuteAction::CreateSP( this, &SRetargetSourceWindow::CanPerformDelete ) );
 		const FText Label = LOCTEXT("DeleteRetargetSourceActionLabel", "Delete");
 		const FText ToolTipText = LOCTEXT("DeleteRetargetSourceActionTooltip", "Deletes the selected retarget sources.");
 		MenuBuilder.AddMenuEntry( Label, ToolTipText, FSlateIcon(), Action);
 	}
 	{
-		FUIAction Action = FUIAction( FExecuteAction::CreateSP( this, &SRetargetSourceWindow::OnRefreshRetargetSource, false ), 
+		FUIAction Action = FUIAction( FExecuteAction::CreateSP( const_cast<SRetargetSourceWindow*>(this), &SRetargetSourceWindow::OnRefreshRetargetSource, false ), 
 			FCanExecuteAction::CreateSP( this, &SRetargetSourceWindow::CanPerformRefresh ) );
 		const FText Label = LOCTEXT("RefreshRetargetSourceActionLabel", "Update");
 		const FText ToolTipText = LOCTEXT("RefreshRetargetSourceActionTooltip", "Updates the selected retarget sources from source mesh.");
diff --git a/Engine/Source/Editor/Persona/Private/SSkeletonAnimNotifies.cpp b/Engine/Source/Editor/Persona/Private/SSkeletonAnimNotifies.cpp
--- a/Engine/Source/Editor/Persona/Private/SSkeletonAnimNotifies.cpp
+++ b/Engine/Source/Editor/Persona/Private/SSkeletonAnimNotifies.cpp
@@ -166,7 +166,7 @@ TSharedPtr<SWidget> SSkeletonAnimNotifies::OnGetContextMenuContent() const
 	{
 		MenuBuilder.BeginSection("AnimNotifyAction", LOCTEXT("SelectedSyncMarkerActions", "Selected Sync Marker Actions"));
 		{
-			FUIAction Action = FUIAction(FExecuteAction::CreateSP(this, &SSkeletonAnimNotifies::OnDeleteSyncMarker));
+			FUIAction Action = FUIAction(FExecuteAction::CreateSP(const_cast<SSkeletonAnimNotifies*>(this), &SSkeletonAnimNotifies::OnDeleteSyncMarker));
 			const FText Label = LOCTEXT("DeleteSyncMarkerButtonLabel", "Delete");
 			const FText ToolTipText = LOCTEXT("DeleteSyncMarkerButtonTooltip", "Deletes the sync marker from the suggestions");
 			MenuBuilder.AddMenuEntry(Label, ToolTipText, FSlateIcon(), Action);
@@ -177,7 +177,7 @@ TSharedPtr<SWidget> SSkeletonAnimNotifies::OnGetContextMenuContent() const
 	{
 		MenuBuilder.BeginSection("AnimNotifyAction", LOCTEXT("AnimNotifyActions", "Notifies"));
 		{
-			FUIAction Action = FUIAction(FExecuteAction::CreateSP(this, &SSkeletonAnimNotifies::OnAddAnimNotify));
+			FUIAction Action = FUIAction(FExecuteAction::CreateSP(const_cast<SSkeletonAnimNotifies*>(this), &SSkeletonAnimNotifies::OnAddAnimNotify));
 			const FText Label = LOCTEXT("NewAnimNotifyButtonLabel", "New...");
 			const FText ToolTipText = LOCTEXT("NewAnimNotifyButtonTooltip", "Creates a new anim notify.");
 			MenuBuilder.AddMenuEntry(Label, ToolTipText, FSlateIcon(), Action);
@@ -187,7 +187,7 @@ TSharedPtr<SWidget> SSkeletonAnimNotifies::OnGetContextMenuContent() const
 		MenuBuilder.BeginSection("AnimNotifyAction", LOCTEXT("SelectedAnimNotifyActions", "Selected Notify Actions"));
 		{
 			{
-				FUIAction Action = FUIAction(FExecuteAction::CreateSP(this, &SSkeletonAnimNotifies::OnRenameAnimNotify),
+				FUIAction Action = FUIAction(FExecuteAction::CreateSP(const_cast<SSkeletonAnimNotifies*>(this), &SSkeletonAnimNotifies::OnRenameAnimNotify),
 					FCanExecuteAction::CreateSP(this, &SSkeletonAnimNotifies::CanPerformRename));
 				const FText Label = LOCTEXT("RenameAnimNotifyButtonLabel", "Rename");
 				const FText ToolTipText = LOCTEXT("RenameAnimNotifyButtonTooltip", "Renames the selected anim notifies.");
@@ -195,7 +195,7 @@ TSharedPtr<SWidget> SSkeletonAnimNotifies::OnGetContextMenuContent() const
 			}
 
 			{
-				FUIAction Action = FUIAction(FExecuteAction::CreateSP(this, &SSkeletonAnimNotifies::OnDeleteAnimNotify),
+				FUIAction Action = FUIAction(FExecuteAction::CreateSP(const_cast<SSkeletonAnimNotifies*>(this), &SSkeletonAnimNotifies::OnDeleteAnimNotify),
 					FCanExecuteAction::CreateSP(this, &SSkeletonAnimNotifies::CanPerformDelete));
 				const FText Label = LOCTEXT("DeleteAnimNotifyButtonLabel", "Delete");
 				const FText ToolTipText = LOCTEXT("DeleteAnimNotifyButtonTooltip", "Deletes the selected anim notifies.");
diff --git a/Engine/Source/Editor/Persona/Private/SSkeletonSlotNames.cpp b/Engine/Source/Editor/Persona/Private/SSkeletonSlotNames.cpp
--- a/Engine/Source/Editor/Persona/Private/SSkeletonSlotNames.cpp
+++ b/Engine/Source/Editor/Persona/Private/SSkeletonSlotNames.cpp
@@ -258,8 +258,8 @@ TSharedPtr<SWidget> SSkeletonSlotNames::OnGetContextMenuContent() const
 		{
 			FDisplayedSlotNameInfo* SlotInfo = SelectedItemPtr.Get();
 
-			FUIAction Action = FUIAction(FExecuteAction::CreateSP(this, &SSkeletonSlotNames::OnDeleteSlotGroup, SlotInfo->Name));
-			Action.CanExecuteAction.BindSP(this, &SSkeletonSlotNames::CanDeleteSlotGroup, SlotInfo->Name);
+			FUIAction Action = FUIAction(FExecuteAction::CreateSP(const_cast<SSkeletonSlotNames*>(this), &SSkeletonSlotNames::OnDeleteSlotGroup, SlotInfo->Name));
+			Action.CanExecuteAction.BindSP(const_cast<SSkeletonSlotNames*>(this), &SSkeletonSlotNames::CanDeleteSlotGroup, SlotInfo->Name);
 			const FText Label = LOCTEXT("AnimSlotManagerContextMenuDeleteSlotGroupLabel", "Delete Slot Group");
 			const FText ToolTipText = LOCTEXT("AnimSlotManagerContextMenuDeleteSlotGroupTooltip", "Delete this slot group.");
 			MenuBuilder.AddMenuEntry(Label, ToolTipText, FSlateIcon(), Action);
@@ -276,13 +276,13 @@ TSharedPtr<SWidget> SSkeletonSlotNames::OnGetContextMenuContent() const
 			MenuBuilder.AddSubMenu(
 				FText::Format(LOCTEXT("ContextMenuSetSlotGroupLabel", "Set Slot {0} Group to"), FText::FromName(SelectedItems[0].Get()->Name)),
 				FText::Format(LOCTEXT("ContextMenuSetSlotGroupToolTip", "Set Slot {0} Group"), FText::FromName(SelectedItems[0].Get()->Name)),
-				FNewMenuDelegate::CreateRaw(this, &SSkeletonSlotNames::FillSetSlotGroupSubMenu));
+				FNewMenuDelegate::CreateRaw(const_cast<SSkeletonSlotNames*>(this), &SSkeletonSlotNames::FillSetSlotGroupSubMenu));
 		}
 		// Rename Slot
 		{
 			FDisplayedSlotNameInfo* SlotInfo = SelectedItemPtr.Get();
 
-			FUIAction Action = FUIAction(FExecuteAction::CreateSP(this, &SSkeletonSlotNames::OnRenameSlot, SlotInfo->Name));
+			FUIAction Action = FUIAction(FExecuteAction::CreateSP(const_cast<SSkeletonSlotNames*>(this), &SSkeletonSlotNames::OnRenameSlot, SlotInfo->Name));
 			const FText Label = LOCTEXT("AnimSlotManagerContextMenuRenameSlotLabel", "Rename Slot");
 			const FText ToolTipText = LOCTEXT("AnimSlotManagerContextMenuRenameSlotTooltip", "Rename this slot");
 			MenuBuilder.AddMenuEntry(Label, ToolTipText, FSlateIcon(), Action);
@@ -291,7 +291,7 @@ TSharedPtr<SWidget> SSkeletonSlotNames::OnGetContextMenuContent() const
 		{
 			FDisplayedSlotNameInfo* SlotInfo = SelectedItemPtr.Get();
 
-			FUIAction Action = FUIAction(FExecuteAction::CreateSP(this, &SSkeletonSlotNames::OnDeleteSlot, SlotInfo->Name));
+			FUIAction Action = FUIAction(FExecuteAction::CreateSP(const_cast<SSkeletonSlotNames*>(this), &SSkeletonSlotNames::OnDeleteSlot, SlotInfo->Name));
 			const FText Label = LOCTEXT("AnimSlotManagerContextMenuDeleteSlotLabel", "Delete Slot");
 			const FText ToolTipText = LOCTEXT("AnimSlotManagerContextMenuDeleteSlotTooltip", "Delete this slot.");
 			MenuBuilder.AddMenuEntry(Label, ToolTipText, FSlateIcon(), Action);
@@ -302,14 +302,14 @@ TSharedPtr<SWidget> SSkeletonSlotNames::OnGetContextMenuContent() const
 	MenuBuilder.BeginSection("SlotManagerGeneralActions", LOCTEXT("SlotManagerGeneralActions", "Slot Manager Actions"));
 	// Add Slot
 	{
-		FUIAction Action = FUIAction(FExecuteAction::CreateSP(this, &SSkeletonSlotNames::OnAddSlot));
+		FUIAction Action = FUIAction(FExecuteAction::CreateSP(const_cast<SSkeletonSlotNames*>(this), &SSkeletonSlotNames::OnAddSlot));
 		const FText Label = LOCTEXT("AnimSlotManagerContextMenuAddSlotLabel", "Add Slot");
 		const FText ToolTipText = LOCTEXT("AnimSlotManagerContextMenuAddSlotTooltip", "Adds a new Slot");
 		MenuBuilder.AddMenuEntry(Label, ToolTipText, FSlateIcon(), Action);
 	}
 	// Add Group
 	{
-		FUIAction Action = FUIAction(FExecuteAction::CreateSP(this, &SSkeletonSlotNames::OnAddGroup));
+		FUIAction Action = FUIAction(FExecuteAction::CreateSP(const_cast<SSkeletonSlotNames*>(this), &SSkeletonSlotNames::OnAddGroup));
 		const FText Label = LOCTEXT("AnimSlotManagerContextMenuAddGroupLabel", "Add Group");
 		const FText ToolTipText = LOCTEXT("AnimSlotManagerContextMenuAddGroupTooltip", "Adds a new Group");
 		MenuBuilder.AddMenuEntry(Label, ToolTipText, FSlateIcon(), Action);
diff --git a/Engine/Source/Editor/Persona/Private/TabSpawners.cpp b/Engine/Source/Editor/Persona/Private/TabSpawners.cpp
--- a/Engine/Source/Editor/Persona/Private/TabSpawners.cpp
+++ b/Engine/Source/Editor/Persona/Private/TabSpawners.cpp
@@ -473,7 +473,7 @@ TSharedRef<SWidget> FAnimBlueprintPreviewEditorSummoner::CreateTabBody(const FWo
 						SNew(SCheckBox)
 						.Style(FEditorStyle::Get(), "RadioButton")
 						.IsChecked(this, &FAnimBlueprintPreviewEditorSummoner::IsChecked, EAnimBlueprintEditorMode::PreviewMode)
-						.OnCheckStateChanged(this, &FAnimBlueprintPreviewEditorSummoner::OnCheckedChanged, EAnimBlueprintEditorMode::PreviewMode)
+						.OnCheckStateChanged(const_cast<FAnimBlueprintPreviewEditorSummoner*>(this), &FAnimBlueprintPreviewEditorSummoner::OnCheckedChanged, EAnimBlueprintEditorMode::PreviewMode)
 						.ToolTip(IDocumentation::Get()->CreateToolTip(	LOCTEXT("AnimBlueprintPropertyEditorPreviewMode", "Switch to editing the preview instance properties"),
 																		NULL,
 																		TEXT("Shared/Editors/Persona"),
@@ -495,7 +495,7 @@ TSharedRef<SWidget> FAnimBlueprintPreviewEditorSummoner::CreateTabBody(const FWo
 						SNew(SCheckBox)
 						.Style(FEditorStyle::Get(), "RadioButton")
 						.IsChecked(this, &FAnimBlueprintPreviewEditorSummoner::IsChecked, EAnimBlueprintEditorMode::DefaultsMode)
-						.OnCheckStateChanged(this, &FAnimBlueprintPreviewEditorSummoner::OnCheckedChanged, EAnimBlueprintEditorMode::DefaultsMode)
+						.OnCheckStateChanged(const_cast<FAnimBlueprintPreviewEditorSummoner*>(this), &FAnimBlueprintPreviewEditorSummoner::OnCheckedChanged, EAnimBlueprintEditorMode::DefaultsMode)
 						.ToolTip(IDocumentation::Get()->CreateToolTip(	LOCTEXT("AnimBlueprintPropertyEditorDefaultMode", "Switch to editing the class defaults"),
 																		NULL,
 																		TEXT("Shared/Editors/Persona"),
@@ -630,8 +630,8 @@ TSharedRef<SWidget> FAdvancedPreviewSceneTabSummoner::CreateTabBody(const FWorkf
 	TArray<FAdvancedPreviewSceneModule::FDetailCustomizationInfo> DetailsCustomizations;
 	TArray<FAdvancedPreviewSceneModule::FPropertyTypeCustomizationInfo> PropertyTypeCustomizations;
 
-	DetailsCustomizations.Add({ UPersonaPreviewSceneDescription::StaticClass(), FOnGetDetailCustomizationInstance::CreateSP(this, &FAdvancedPreviewSceneTabSummoner::CustomizePreviewSceneDescription) });
-	PropertyTypeCustomizations.Add({ FPreviewMeshCollectionEntry::StaticStruct()->GetFName(), FOnGetPropertyTypeCustomizationInstance::CreateSP(this, &FAdvancedPreviewSceneTabSummoner::CustomizePreviewMeshCollectionEntry) });
+	DetailsCustomizations.Add({ UPersonaPreviewSceneDescription::StaticClass(), FOnGetDetailCustomizationInstance::CreateSP(const_cast<FAdvancedPreviewSceneTabSummoner*>(this), &FAdvancedPreviewSceneTabSummoner::CustomizePreviewSceneDescription) });
+	PropertyTypeCustomizations.Add({ FPreviewMeshCollectionEntry::StaticStruct()->GetFName(), FOnGetPropertyTypeCustomizationInstance::CreateSP(const_cast<FAdvancedPreviewSceneTabSummoner*>(this), &FAdvancedPreviewSceneTabSummoner::CustomizePreviewMeshCollectionEntry) });
 
 	FAdvancedPreviewSceneModule& AdvancedPreviewSceneModule = FModuleManager::LoadModuleChecked<FAdvancedPreviewSceneModule>("AdvancedPreviewScene");
 	return AdvancedPreviewSceneModule.CreateAdvancedPreviewSceneSettingsWidget(PreviewSceneRef, PreviewSceneRef->GetPreviewSceneDescription(), DetailsCustomizations, PropertyTypeCustomizations);
diff --git a/Engine/Source/Editor/SceneOutliner/Private/SSceneOutliner.cpp b/Engine/Source/Editor/SceneOutliner/Private/SSceneOutliner.cpp
--- a/Engine/Source/Editor/SceneOutliner/Private/SSceneOutliner.cpp
+++ b/Engine/Source/Editor/SceneOutliner/Private/SSceneOutliner.cpp
@@ -1667,7 +1667,7 @@ namespace SceneOutliner
 	void SSceneOutliner::FillFoldersSubMenu(FMenuBuilder& MenuBuilder) const
 	{
 		MenuBuilder.AddMenuEntry(LOCTEXT( "CreateNew", "Create New Folder" ), LOCTEXT( "CreateNew_ToolTip", "Move to a new folder" ),
-			FSlateIcon(FEditorStyle::GetStyleSetName(), "SceneOutliner.NewFolderIcon"), FExecuteAction::CreateSP(this, &SSceneOutliner::CreateFolder));
+			FSlateIcon(FEditorStyle::GetStyleSetName(), "SceneOutliner.NewFolderIcon"), FExecuteAction::CreateSP(const_cast<SSceneOutliner*>(this), &SSceneOutliner::CreateFolder));
 
 		AddMoveToFolderOutliner(MenuBuilder);
 	}
@@ -1793,7 +1793,7 @@ namespace SceneOutliner
 			[
 				SNew(SSceneOutliner, MiniSceneOutlinerInitOptions)
 				.IsEnabled(FSlateApplication::Get().GetNormalExecutionAttribute())
-				.OnItemPickedDelegate(FOnSceneOutlinerItemPicked::CreateSP(this, &SSceneOutliner::MoveSelectionTo))
+				.OnItemPickedDelegate(FOnSceneOutlinerItemPicked::CreateSP(const_cast<SSceneOutliner*>(this), &SSceneOutliner::MoveSelectionTo))
 			];
 
 		MenuBuilder.BeginSection(FName(), LOCTEXT("ExistingFolders", "Existing:"));
@@ -1807,12 +1807,12 @@ namespace SceneOutliner
 			LOCTEXT( "AddChildrenToSelection", "Immediate Children" ),
 			LOCTEXT( "AddChildrenToSelection_ToolTip", "Select all immediate actor children of the selected folders" ),
 			FSlateIcon(),
-			FExecuteAction::CreateSP(this, &SSceneOutliner::SelectFoldersDescendants, /*bSelectImmediateChildrenOnly=*/ true));
+			FExecuteAction::CreateSP(const_cast<SSceneOutliner*>(this), &SSceneOutliner::SelectFoldersDescendants, /*bSelectImmediateChildrenOnly=*/ true));
 		MenuBuilder.AddMenuEntry(
 			LOCTEXT( "AddDescendantsToSelection", "All Descendants" ),
 			LOCTEXT( "AddDescendantsToSelection_ToolTip", "Select all actor descendants of the selected folders" ),
 			FSlateIcon(),
-			FExecuteAction::CreateSP(this, &SSceneOutliner::SelectFoldersDescendants, /*bSelectImmediateChildrenOnly=*/ false));
+			FExecuteAction::CreateSP(const_cast<SSceneOutliner*>(this), &SSceneOutliner::SelectFoldersDescendants, /*bSelectImmediateChildrenOnly=*/ false));
 	}
 
 	void SSceneOutliner::SelectFoldersDescendants(bool bSelectImmediateChildrenOnly)
diff --git a/Engine/Source/Editor/StaticMeshEditor/Private/StaticMeshEditorTools.cpp b/Engine/Source/Editor/StaticMeshEditor/Private/StaticMeshEditorTools.cpp
--- a/Engine/Source/Editor/StaticMeshEditor/Private/StaticMeshEditorTools.cpp
+++ b/Engine/Source/Editor/StaticMeshEditor/Private/StaticMeshEditorTools.cpp
@@ -3576,8 +3576,8 @@ TSharedRef<SWidget> FLevelOfDetailSettingsLayout::GetLODScreenSizeWidget(FName P
 		.MaxValue(WORLD_MAX)
 		.SliderExponent(2.0f)
 		.Value(this, &FLevelOfDetailSettingsLayout::GetLODScreenSize, PlatformGroupName, LODIndex)
-		.OnValueChanged(this, &FLevelOfDetailSettingsLayout::OnLODScreenSizeChanged, PlatformGroupName, LODIndex)
-		.OnValueCommitted(this, &FLevelOfDetailSettingsLayout::OnLODScreenSizeCommitted, PlatformGroupName, LODIndex)
+		.OnValueChanged(const_cast<FLevelOfDetailSettingsLayout*>(this), &FLevelOfDetailSettingsLayout::OnLODScreenSizeChanged, PlatformGroupName, LODIndex)
+		.OnValueCommitted(const_cast<FLevelOfDetailSettingsLayout*>(this), &FLevelOfDetailSettingsLayout::OnLODScreenSizeCommitted, PlatformGroupName, LODIndex)
 		.IsEnabled(this, &FLevelOfDetailSettingsLayout::CanChangeLODScreenSize);
 }
 
@@ -4014,8 +4014,8 @@ TSharedRef<SWidget> FLevelOfDetailSettingsLayout::GetMinLODWidget(FName Platform
 	return SNew(SSpinBox<int32>)
 		.Font(IDetailLayoutBuilder::GetDetailFont())
 		.Value(this, &FLevelOfDetailSettingsLayout::GetMinLOD, PlatformGroupName)
-		.OnValueChanged(this, &FLevelOfDetailSettingsLayout::OnMinLODChanged, PlatformGroupName)
-		.OnValueCommitted(this, &FLevelOfDetailSettingsLayout::OnMinLODCommitted, PlatformGroupName)
+		.OnValueChanged(const_cast<FLevelOfDetailSettingsLayout*>(this), &FLevelOfDetailSettingsLayout::OnMinLODChanged, PlatformGroupName)
+		.OnValueCommitted(const_cast<FLevelOfDetailSettingsLayout*>(this), &FLevelOfDetailSettingsLayout::OnMinLODCommitted, PlatformGroupName)
 		.MinValue(0)
 		.MaxValue(MAX_STATIC_MESH_LODS)
 		.ToolTipText(this, &FLevelOfDetailSettingsLayout::GetMinLODTooltip)
diff --git a/Engine/Source/Editor/StatsViewer/Private/SStatsViewer.cpp b/Engine/Source/Editor/StatsViewer/Private/SStatsViewer.cpp
--- a/Engine/Source/Editor/StatsViewer/Private/SStatsViewer.cpp
+++ b/Engine/Source/Editor/StatsViewer/Private/SStatsViewer.cpp
@@ -615,7 +615,7 @@ TSharedRef<SWidget> SStatsViewer::OnGetDisplayMenuContent() const
 			StatsPage->GetToolTip(), 
 			FSlateIcon(), 
 			FUIAction( 
-				FExecuteAction::CreateSP( this, &SStatsViewer::SetDisplayedStats, StatsPage ),
+				FExecuteAction::CreateSP( const_cast<SStatsViewer*>(this), &SStatsViewer::SetDisplayedStats, StatsPage ),
 				FCanExecuteAction(),
 				FIsActionChecked::CreateSP( this, &SStatsViewer::AreStatsDisplayed, StatsPage )
 			),
@@ -640,7 +640,7 @@ TSharedRef<SWidget> SStatsViewer::OnGetObjectSetMenuContent() const
 				FText::FromString( CurrentStats->GetObjectSetToolTip( ObjectSetIndex ) ), 
 				FSlateIcon(), 
 				FUIAction( 
-					FExecuteAction::CreateSP( this, &SStatsViewer::SetObjectSet, ObjectSetIndex ),
+					FExecuteAction::CreateSP( const_cast<SStatsViewer*>(this), &SStatsViewer::SetObjectSet, ObjectSetIndex ),
 					FCanExecuteAction(),
 					FIsActionChecked::CreateSP( this, &SStatsViewer::IsObjectSetSelected, ObjectSetIndex )
 				),
@@ -686,7 +686,7 @@ TSharedRef<SWidget> SStatsViewer::OnGetFilterMenuContent() const
 					FText::Format( LOCTEXT( "FilterMenuEntry_Tooltip", "Search statistics by {FilterName}.\n{FilterDesc}" ), Arguments ), 
 					FSlateIcon(), 
 					FUIAction( 
-						FExecuteAction::CreateSP( this, &SStatsViewer::SetSearchFilter, ColumnIndex ),
+						FExecuteAction::CreateSP( const_cast<SStatsViewer*>(this), &SStatsViewer::SetSearchFilter, ColumnIndex ),
 						FCanExecuteAction(),
 						FIsActionChecked::CreateSP( this, &SStatsViewer::IsSearchFilterSelected, ColumnIndex )
 					),
diff --git a/Engine/Source/Editor/StatsViewer/Private/StatsPages/CookerStatsPage.cpp b/Engine/Source/Editor/StatsViewer/Private/StatsPages/CookerStatsPage.cpp
--- a/Engine/Source/Editor/StatsViewer/Private/StatsPages/CookerStatsPage.cpp
+++ b/Engine/Source/Editor/StatsViewer/Private/StatsPages/CookerStatsPage.cpp
@@ -180,7 +180,7 @@ TSharedRef<SWidget> FCookerStatsPage::HandleFilterComboButtonGetMenuContent( ) c
 					FText::FromString(PlatformName),
 					FSlateIcon(),
 					FUIAction(
-						FExecuteAction::CreateRaw(&CookerStatsPage, &FCookerStatsPage::HandleFilterMenuEntryExecute, PlatformName),
+						FExecuteAction::CreateRaw(const_cast<FCookerStatsPage*>(&CookerStatsPage), &FCookerStatsPage::HandleFilterMenuEntryExecute, PlatformName),
 						FCanExecuteAction(),
 						FIsActionChecked::CreateRaw(&CookerStatsPage, &FCookerStatsPage::HandleFilterMenuEntryIsChecked, PlatformName)
 					),
diff --git a/Engine/Source/Editor/UndoHistory/Private/Widgets/SUndoHistory.cpp b/Engine/Source/Editor/UndoHistory/Private/Widgets/SUndoHistory.cpp
--- a/Engine/Source/Editor/UndoHistory/Private/Widgets/SUndoHistory.cpp
+++ b/Engine/Source/Editor/UndoHistory/Private/Widgets/SUndoHistory.cpp
@@ -459,7 +459,7 @@ TSharedRef<SWidget> SUndoHistory::GetViewButtonContent() const
 			LOCTEXT("ToggleShowTransactionDetailsToolTip", "When enabled, display additional information about transactions."),
 			FSlateIcon(),
 			FUIAction(
-				FExecuteAction::CreateSP(this, &SUndoHistory::ToggleShowTransactionDetails),
+				FExecuteAction::CreateSP(const_cast<SUndoHistory*>(this), &SUndoHistory::ToggleShowTransactionDetails),
 				FCanExecuteAction(),
 				FIsActionChecked::CreateSP(this, &SUndoHistory::IsShowingTransactionDetails)
 			),
diff --git a/Engine/Source/Editor/UnrealEd/Private/SCommonEditorViewportToolbarBase.cpp b/Engine/Source/Editor/UnrealEd/Private/SCommonEditorViewportToolbarBase.cpp
--- a/Engine/Source/Editor/UnrealEd/Private/SCommonEditorViewportToolbarBase.cpp
+++ b/Engine/Source/Editor/UnrealEd/Private/SCommonEditorViewportToolbarBase.cpp
@@ -356,7 +356,7 @@ TSharedRef<SWidget> SCommonEditorViewportToolbarBase::GenerateScreenPercentageMe
 		.MinValue(PreviewScreenPercentageMin)
 		.MaxValue(PreviewScreenPercentageMax)
 		.Value(this, &SCommonEditorViewportToolbarBase::OnGetScreenPercentageValue)
-		.OnValueChanged(this, &SCommonEditorViewportToolbarBase::OnScreenPercentageValueChanged)
+		.OnValueChanged(const_cast<SCommonEditorViewportToolbarBase*>(this), &SCommonEditorViewportToolbarBase::OnScreenPercentageValueChanged)
 		]
 		];
 }
@@ -395,7 +395,7 @@ TSharedRef<SWidget> SCommonEditorViewportToolbarBase::GenerateFarViewPlaneMenu()
 				.MaxValue(100000.0f)
 				.Font(FEditorStyle::GetFontStyle(TEXT("MenuItem.Font")))
 				.Value(this, &SCommonEditorViewportToolbarBase::OnGetFarViewPlaneValue)
-				.OnValueChanged(this, &SCommonEditorViewportToolbarBase::OnFarViewPlaneValueChanged)
+				.OnValueChanged(const_cast<SCommonEditorViewportToolbarBase*>(this), &SCommonEditorViewportToolbarBase::OnFarViewPlaneValueChanged)
 			]
 		];
 }
diff --git a/Engine/Source/Editor/UnrealEd/Private/SCurveEditor.cpp b/Engine/Source/Editor/UnrealEd/Private/SCurveEditor.cpp
--- a/Engine/Source/Editor/UnrealEd/Private/SCurveEditor.cpp
+++ b/Engine/Source/Editor/UnrealEd/Private/SCurveEditor.cpp
@@ -611,7 +611,7 @@ TSharedRef<SWidget> SCurveEditor::CreateCurveSelectionWidget() const
 				[
 					SNew(SCheckBox)
 					.IsChecked(this, &SCurveEditor::IsCurveVisible, CurveViewModel)
-					.OnCheckStateChanged(this, &SCurveEditor::OnCurveIsVisibleChanged, CurveViewModel)
+					.OnCheckStateChanged(const_cast<SCurveEditor*>(this), &SCurveEditor::OnCurveIsVisibleChanged, CurveViewModel)
 					.ToolTipText(this, &SCurveEditor::GetIsCurveVisibleToolTip, CurveViewModel)
 					.CheckedImage(FEditorStyle::GetBrush("CurveEd.Visible"))
 					.CheckedHoveredImage(FEditorStyle::GetBrush("CurveEd.VisibleHighlight"))
@@ -628,7 +628,7 @@ TSharedRef<SWidget> SCurveEditor::CreateCurveSelectionWidget() const
 				[
 					SNew(SCheckBox)
 					.IsChecked(this, &SCurveEditor::IsCurveLocked, CurveViewModel)
-					.OnCheckStateChanged(this, &SCurveEditor::OnCurveIsLockedChanged, CurveViewModel)
+					.OnCheckStateChanged(const_cast<SCurveEditor*>(this), &SCurveEditor::OnCurveIsLockedChanged, CurveViewModel)
 					.ToolTipText(this, &SCurveEditor::GetIsCurveLockedToolTip, CurveViewModel)
 					.CheckedImage(FEditorStyle::GetBrush("CurveEd.Locked"))
 					.CheckedHoveredImage(FEditorStyle::GetBrush("CurveEd.LockedHighlight"))
diff --git a/Engine/Source/Editor/UnrealEd/Private/SEditorViewport.cpp b/Engine/Source/Editor/UnrealEd/Private/SEditorViewport.cpp
--- a/Engine/Source/Editor/UnrealEd/Private/SEditorViewport.cpp
+++ b/Engine/Source/Editor/UnrealEd/Private/SEditorViewport.cpp
@@ -571,7 +571,7 @@ TSharedRef<SWidget> SEditorViewport::BuildFixedEV100Menu()  const
 				.MinValue(EV100Min)
 				.MaxValue(EV100Max)
 				.Value( this, &SEditorViewport::OnGetFixedEV100Value )
-				.OnValueChanged( this, &SEditorViewport::OnFixedEV100ValueChanged )
+				.OnValueChanged( const_cast<SEditorViewport*>(this), &SEditorViewport::OnFixedEV100ValueChanged )
 				.ToolTipText(LOCTEXT( "EV100ToolTip", "Sets the exposure value of the camera using the specified EV100. Exposure = 1 / (1.2 * 2^EV100)"))
 				.IsEnabled( this, &SEditorViewport::IsFixedEV100Enabled )
 			]
diff --git a/Engine/Source/Editor/WorldBrowser/Private/LevelCollectionModel.cpp b/Engine/Source/Editor/WorldBrowser/Private/LevelCollectionModel.cpp
--- a/Engine/Source/Editor/WorldBrowser/Private/LevelCollectionModel.cpp
+++ b/Engine/Source/Editor/WorldBrowser/Private/LevelCollectionModel.cpp
@@ -794,7 +794,7 @@ void FLevelCollectionModel::CustomizeFileMainMenu(FMenuBuilder& InMenuBuilder) c
 	InMenuBuilder.AddSubMenu( 
 		LOCTEXT("SourceControl", "Source Control"),
 		LOCTEXT("SourceControl_ToolTip", "Source Control Options"),
-		FNewMenuDelegate::CreateSP(this, &FLevelCollectionModel::FillSourceControlSubMenu));
+		FNewMenuDelegate::CreateSP(const_cast<FLevelCollectionModel*>(this), &FLevelCollectionModel::FillSourceControlSubMenu));
 		
 	if (AreAnyLevelsSelected())
 	{
diff --git a/Engine/Source/Editor/WorldBrowser/Private/SWorldHierarchyImpl.cpp b/Engine/Source/Editor/WorldBrowser/Private/SWorldHierarchyImpl.cpp
--- a/Engine/Source/Editor/WorldBrowser/Private/SWorldHierarchyImpl.cpp
+++ b/Engine/Source/Editor/WorldBrowser/Private/SWorldHierarchyImpl.cpp
@@ -498,7 +498,7 @@ TSharedPtr<SWidget> SWorldHierarchyImpl::ConstructLevelContextMenu() const
 				MenuBuilder.AddSubMenu(
 					LOCTEXT("MoveSelectionTo", "Move To"),
 					LOCTEXT("MoveSelectionTo_Tooltip", "Move selection to another folder"),
-					FNewMenuDelegate::CreateSP(this, &SWorldHierarchyImpl::FillFoldersSubmenu)
+					FNewMenuDelegate::CreateSP(const_cast<SWorldHierarchyImpl*>(this), &SWorldHierarchyImpl::FillFoldersSubmenu)
 				);
 			}
 
@@ -507,7 +507,7 @@ TSharedPtr<SWidget> SWorldHierarchyImpl::ConstructLevelContextMenu() const
 				MenuBuilder.AddSubMenu(
 					LOCTEXT("SelectSubmenu", "Select"),
 					LOCTEXT("SelectSubmenu_Tooltip", "Select child items of the current selection"),
-					FNewMenuDelegate::CreateSP(this, &SWorldHierarchyImpl::FillSelectionSubmenu)
+					FNewMenuDelegate::CreateSP(const_cast<SWorldHierarchyImpl*>(this), &SWorldHierarchyImpl::FillSelectionSubmenu)
 				);
 			}
 		}
diff --git a/Engine/Source/Editor/WorldBrowser/Private/StreamingLevels/StreamingLevelCollectionModel.cpp b/Engine/Source/Editor/WorldBrowser/Private/StreamingLevels/StreamingLevelCollectionModel.cpp
--- a/Engine/Source/Editor/WorldBrowser/Private/StreamingLevels/StreamingLevelCollectionModel.cpp
+++ b/Engine/Source/Editor/WorldBrowser/Private/StreamingLevels/StreamingLevelCollectionModel.cpp
@@ -288,13 +288,13 @@ void FStreamingLevelCollectionModel::BuildHierarchyMenu(FMenuBuilder& InMenuBuil
 		InMenuBuilder.AddSubMenu( 
 			LOCTEXT("VisibilityHeader", "Visibility"),
 			LOCTEXT("VisibilitySubMenu_ToolTip", "Selected Level(s) visibility commands"),
-			FNewMenuDelegate::CreateSP(this, &FStreamingLevelCollectionModel::FillVisibilitySubMenu ) );
+			FNewMenuDelegate::CreateSP(const_cast<FStreamingLevelCollectionModel*>(this), &FStreamingLevelCollectionModel::FillVisibilitySubMenu ) );
 
 		// Lock commands
 		InMenuBuilder.AddSubMenu( 
 			LOCTEXT("LockHeader", "Lock"),
 			LOCTEXT("LockSubMenu_ToolTip", "Selected Level(s) lock commands"),
-			FNewMenuDelegate::CreateSP(this, &FStreamingLevelCollectionModel::FillLockSubMenu ) );
+			FNewMenuDelegate::CreateSP(const_cast<FStreamingLevelCollectionModel*>(this), &FStreamingLevelCollectionModel::FillLockSubMenu ) );
 		
 		// Level streaming specific commands
 		if (AreAnyLevelsSelected() && !(IsOneLevelSelected() && GetSelectedLevels()[0]->IsPersistent()))
@@ -304,7 +304,7 @@ void FStreamingLevelCollectionModel::BuildHierarchyMenu(FMenuBuilder& InMenuBuil
 			InMenuBuilder.AddSubMenu( 
 				LOCTEXT("LevelsChangeStreamingMethod", "Change Streaming Method"),
 				LOCTEXT("LevelsChangeStreamingMethod_Tooltip", "Changes the streaming method for the selected levels"),
-				FNewMenuDelegate::CreateRaw(this, &FStreamingLevelCollectionModel::FillSetStreamingMethodSubMenu ));
+				FNewMenuDelegate::CreateRaw(const_cast<FStreamingLevelCollectionModel*>(this), &FStreamingLevelCollectionModel::FillSetStreamingMethodSubMenu ));
 		}
 
 		if (IsOneLevelSelected() && !GetSelectedLevels()[0]->IsPersistent())
@@ -312,7 +312,7 @@ void FStreamingLevelCollectionModel::BuildHierarchyMenu(FMenuBuilder& InMenuBuil
 			InMenuBuilder.AddSubMenu( 
 				LOCTEXT("LevelsChangeLightingScenario", "Lighting Scenario"),
 				LOCTEXT("LevelsChangeLightingScenario_Tooltip", "Changes Lighting Scenario Status for the selected level"),
-				FNewMenuDelegate::CreateRaw(this, &FStreamingLevelCollectionModel::FillChangeLightingScenarioSubMenu ));
+				FNewMenuDelegate::CreateRaw(const_cast<FStreamingLevelCollectionModel*>(this), &FStreamingLevelCollectionModel::FillChangeLightingScenarioSubMenu ));
 		}
 
 		InMenuBuilder.AddMenuEntry(Commands.World_FindInContentBrowser);
diff --git a/Engine/Source/Editor/WorldBrowser/Private/StreamingLevels/StreamingLevelCustomization.cpp b/Engine/Source/Editor/WorldBrowser/Private/StreamingLevels/StreamingLevelCustomization.cpp
--- a/Engine/Source/Editor/WorldBrowser/Private/StreamingLevels/StreamingLevelCustomization.cpp
+++ b/Engine/Source/Editor/WorldBrowser/Private/StreamingLevels/StreamingLevelCustomization.cpp
@@ -281,14 +281,14 @@ FUIAction FStreamingLevelCustomization::CreateCopyAction(ELevelTransformField::T
 	return
 		FUIAction
 		(
-			FExecuteAction::CreateSP(this, &FStreamingLevelCustomization::OnCopy, TransformField)
+			FExecuteAction::CreateSP(const_cast<FStreamingLevelCustomization*>(this), &FStreamingLevelCustomization::OnCopy, TransformField)
 		);
 }
 
 FUIAction FStreamingLevelCustomization::CreatePasteAction(ELevelTransformField::Type TransformField) const
 {
 	return
-		FUIAction(FExecuteAction::CreateSP(this, &FStreamingLevelCustomization::OnPaste, TransformField));
+		FUIAction(FExecuteAction::CreateSP(const_cast<FStreamingLevelCustomization*>(this), &FStreamingLevelCustomization::OnPaste, TransformField));
 }
 
 void FStreamingLevelCustomization::OnSetLevelPosition( float NewValue, ETextCommit::Type CommitInfo, int32 Axis )
diff --git a/Engine/Source/Editor/WorldBrowser/Private/Tiles/WorldTileCollectionModel.cpp b/Engine/Source/Editor/WorldBrowser/Private/Tiles/WorldTileCollectionModel.cpp
--- a/Engine/Source/Editor/WorldBrowser/Private/Tiles/WorldTileCollectionModel.cpp
+++ b/Engine/Source/Editor/WorldBrowser/Private/Tiles/WorldTileCollectionModel.cpp
@@ -344,13 +344,13 @@ void FWorldTileCollectionModel::BuildWorldCompositionMenu(FMenuBuilder& InMenuBu
 			InMenuBuilder.AddSubMenu( 
 				LOCTEXT("VisibilityHeader", "Visibility"),
 				LOCTEXT("VisibilitySubMenu_ToolTip", "Selected Level(s) visibility commands"),
-				FNewMenuDelegate::CreateSP(this, &FWorldTileCollectionModel::FillVisibilitySubMenu ) );
+				FNewMenuDelegate::CreateSP(const_cast<FWorldTileCollectionModel*>(this), &FWorldTileCollectionModel::FillVisibilitySubMenu ) );
 
 			// Lock commands
 			InMenuBuilder.AddSubMenu( 
 				LOCTEXT("LockHeader", "Lock"),
 				LOCTEXT("LockSubMenu_ToolTip", "Selected Level(s) lock commands"),
-				FNewMenuDelegate::CreateSP(this, &FWorldTileCollectionModel::FillLockSubMenu ) );
+				FNewMenuDelegate::CreateSP(const_cast<FWorldTileCollectionModel*>(this), &FWorldTileCollectionModel::FillLockSubMenu ) );
 
 			InMenuBuilder.AddMenuEntry(Commands.World_FindInContentBrowser);
 		}
@@ -364,7 +364,7 @@ void FWorldTileCollectionModel::BuildWorldCompositionMenu(FMenuBuilder& InMenuBu
 				InMenuBuilder.AddSubMenu( 
 					LOCTEXT("Layer_Assign", "Assign to Layer"),
 					FText::GetEmpty(),
-					FNewMenuDelegate::CreateSP(this, &FWorldTileCollectionModel::FillLayersSubMenu));
+					FNewMenuDelegate::CreateSP(const_cast<FWorldTileCollectionModel*>(this), &FWorldTileCollectionModel::FillLayersSubMenu));
 			}
 			InMenuBuilder.EndSection();
 		}
@@ -410,7 +410,7 @@ void FWorldTileCollectionModel::BuildWorldCompositionMenu(FMenuBuilder& InMenuBu
 				InMenuBuilder.AddSubMenu( 
 					LOCTEXT("AddLandscapeLevel", "Add Adjacent Landscape Level"),
 					FText::GetEmpty(),
-					FNewMenuDelegate::CreateSP(this, &FWorldTileCollectionModel::FillAdjacentLandscapeSubMenu));
+					FNewMenuDelegate::CreateSP(const_cast<FWorldTileCollectionModel*>(this), &FWorldTileCollectionModel::FillAdjacentLandscapeSubMenu));
 			}
 
 			// Tiled landscape
@@ -419,7 +419,7 @@ void FWorldTileCollectionModel::BuildWorldCompositionMenu(FMenuBuilder& InMenuBu
 				InMenuBuilder.AddSubMenu( 
 					LOCTEXT("ReimportTiledLandscape", "Reimport Tiled Landscape"),
 					FText::GetEmpty(),
-					FNewMenuDelegate::CreateSP(this, &FWorldTileCollectionModel::FillReimportTiledLandscapeSubMenu));
+					FNewMenuDelegate::CreateSP(const_cast<FWorldTileCollectionModel*>(this), &FWorldTileCollectionModel::FillReimportTiledLandscapeSubMenu));
 			}
 		
 			InMenuBuilder.EndSection();
@@ -454,13 +454,13 @@ void FWorldTileCollectionModel::BuildHierarchyMenu(FMenuBuilder& InMenuBuilder)
 		InMenuBuilder.AddSubMenu( 
 			LOCTEXT("VisibilityHeader", "Visibility"),
 			LOCTEXT("VisibilitySubMenu_ToolTip", "Selected Level(s) visibility commands"),
-			FNewMenuDelegate::CreateSP(this, &FWorldTileCollectionModel::FillVisibilitySubMenu ) );
+			FNewMenuDelegate::CreateSP(const_cast<FWorldTileCollectionModel*>(this), &FWorldTileCollectionModel::FillVisibilitySubMenu ) );
 
 		// Lock commands
 		InMenuBuilder.AddSubMenu( 
 			LOCTEXT("LockHeader", "Lock"),
 			LOCTEXT("LockSubMenu_ToolTip", "Selected Level(s) lock commands"),
-			FNewMenuDelegate::CreateSP(this, &FWorldTileCollectionModel::FillLockSubMenu ) );
+			FNewMenuDelegate::CreateSP(const_cast<FWorldTileCollectionModel*>(this), &FWorldTileCollectionModel::FillLockSubMenu ) );
 	
 		InMenuBuilder.AddMenuEntry(Commands.World_FindInContentBrowser);
 	}
@@ -474,7 +474,7 @@ void FWorldTileCollectionModel::BuildHierarchyMenu(FMenuBuilder& InMenuBuilder)
 			InMenuBuilder.AddSubMenu( 
 				LOCTEXT("WorldLayers", "Assign to Layer"),
 				FText::GetEmpty(),
-				FNewMenuDelegate::CreateSP(this, &FWorldTileCollectionModel::FillLayersSubMenu)
+				FNewMenuDelegate::CreateSP(const_cast<FWorldTileCollectionModel*>(this), &FWorldTileCollectionModel::FillLayersSubMenu)
 			);
 		}
 		InMenuBuilder.EndSection();
@@ -524,7 +524,7 @@ void FWorldTileCollectionModel::FillLayersSubMenu(FMenuBuilder& InMenuBuilder) c
 			FText::FromString((*It).Name), 
 			FText(), FSlateIcon(),
 			FUIAction(FExecuteAction::CreateSP(
-				this, &FWorldTileCollectionModel::AssignSelectedLevelsToLayer_Executed, (*It)
+				const_cast<FWorldTileCollectionModel*>(this), &FWorldTileCollectionModel::AssignSelectedLevelsToLayer_Executed, (*It)
 				)
 			)
 		);
@@ -550,7 +550,7 @@ void FWorldTileCollectionModel::FillReimportTiledLandscapeSubMenu(FMenuBuilder&
 			LOCTEXT("Menu_HeightmapTitle", "Heightmap"), 
 			FText(), FSlateIcon(),
 			FUIAction(FExecuteAction::CreateSP(
-				this, &FWorldTileCollectionModel::ReimportTiledLandscape_Executed, HeightmapLayerName
+				const_cast<FWorldTileCollectionModel*>(this), &FWorldTileCollectionModel::ReimportTiledLandscape_Executed, HeightmapLayerName
 				)
 			)
 		);
@@ -570,7 +570,7 @@ void FWorldTileCollectionModel::FillWeightmapsSubMenu(FMenuBuilder& InMenuBuilde
 			LOCTEXT("Menu_AllWeightmapsTitle", "All Weightmaps"), 
 			FText(), FSlateIcon(),
 			FUIAction(FExecuteAction::CreateSP(
-				this, &FWorldTileCollectionModel::ReimportTiledLandscape_Executed, FName(NAME_None)
+				const_cast<FWorldTileCollectionModel*>(this), &FWorldTileCollectionModel::ReimportTiledLandscape_Executed, FName(NAME_None)
 				)
 			)
 		);
@@ -596,7 +596,7 @@ void FWorldTileCollectionModel::FillWeightmapsSubMenu(FMenuBuilder& InMenuBuilde
 			FText::FromName(LayerName), 
 			FText(), FSlateIcon(),
 			FUIAction(FExecuteAction::CreateSP(
-				this, &FWorldTileCollectionModel::ReimportTiledLandscape_Executed, LayerName
+				const_cast<FWorldTileCollectionModel*>(this), &FWorldTileCollectionModel::ReimportTiledLandscape_Executed, LayerName
 				)
 			)
 		);
diff --git a/Engine/Source/Editor/WorldBrowser/Private/WorldTreeItemTypes.cpp b/Engine/Source/Editor/WorldBrowser/Private/WorldTreeItemTypes.cpp
--- a/Engine/Source/Editor/WorldBrowser/Private/WorldTreeItemTypes.cpp
+++ b/Engine/Source/Editor/WorldBrowser/Private/WorldTreeItemTypes.cpp
@@ -484,7 +484,7 @@ namespace WorldHierarchy
 			const FSlateIcon NewFolderIcon(FEditorStyle::GetStyleSetName(), "WorldBrowser.NewFolderIcon");
 
 			FName RootPath = LevelModel.Pin()->GetFolderPath();
-			auto NewFolderAction = FExecuteAction::CreateSP(&Hierarchy, &SWorldHierarchyImpl::CreateFolder, LevelModel.Pin(), RootPath);
+			auto NewFolderAction = FExecuteAction::CreateSP(const_cast<SWorldHierarchyImpl*>(&Hierarchy), &SWorldHierarchyImpl::CreateFolder, LevelModel.Pin(), RootPath);
 
 			MenuBuilder.AddMenuEntry(LOCTEXT("CreateFolder", "Create Folder"), FText(), NewFolderIcon, FUIAction(NewFolderAction));
 		}
@@ -912,9 +912,9 @@ namespace WorldHierarchy
 		TArray<FWorldTreeItemPtr> Folders;
 		Folders.Add(AsShared());
 
-		auto NewFolderAction = FExecuteAction::CreateSP(&Hierarchy, &SWorldHierarchyImpl::CreateFolder, RootLevel, Path);
-		auto RenameFolderAction = FExecuteAction::CreateSP(&Hierarchy, &SWorldHierarchyImpl::InitiateRename, AsShared());
-		auto DeleteFolderAction = FExecuteAction::CreateSP(&Hierarchy, &SWorldHierarchyImpl::DeleteFolders, Folders, /*bTransactional*/ true);
+		auto NewFolderAction = FExecuteAction::CreateSP(const_cast<SWorldHierarchyImpl*>(&Hierarchy), &SWorldHierarchyImpl::CreateFolder, RootLevel, Path);
+		auto RenameFolderAction = FExecuteAction::CreateSP(const_cast<SWorldHierarchyImpl*>(&Hierarchy), &SWorldHierarchyImpl::InitiateRename, AsShared());
+		auto DeleteFolderAction = FExecuteAction::CreateSP(const_cast<SWorldHierarchyImpl*>(&Hierarchy), &SWorldHierarchyImpl::DeleteFolders, Folders, /*bTransactional*/ true);
 
 		MenuBuilder.AddMenuEntry(LOCTEXT("CreateSubFolder", "Create Subfolder"), FText(), NewFolderIcon, FUIAction(NewFolderAction));
 		MenuBuilder.AddMenuEntry(LOCTEXT("RenameFolder", "Rename"), FText(), FSlateIcon(FEditorStyle::GetStyleSetName(), "ContentBrowser.AssetActions.Rename"), FUIAction(RenameFolderAction));
