Fixed crash in WorldDetails accessing invalid Level

diff --git a/Engine/Source/Editor/WorldBrowser/Private/SWorldDetails.cpp b/Engine/Source/Editor/WorldBrowser/Private/SWorldDetails.cpp
--- a/Engine/Source/Editor/WorldBrowser/Private/SWorldDetails.cpp
+++ b/Engine/Source/Editor/WorldBrowser/Private/SWorldDetails.cpp
@@ -210,7 +210,8 @@ void SWorldDetails::OnSelectionChanged()
 	else
 	{
 		SubLevelsComboBox->SetSelectedItem(SelectedLevels[0]);
-		WorldDetailsView->SetObject(Cast<UObject>(SelectedLevels[0]->GetLevelObject()->GetLevelPartition()));
+		ULevel* LevelObject = SelectedLevels[0]->GetLevelObject();
+		WorldDetailsView->SetObject(Cast<UObject>(LevelObject ? LevelObject->GetLevelPartition() : nullptr));
 	}
 
 	bUpdatingSelection = false;
