#UE4 Give FLevelFolders a FGCObject referencer name so it shows up in debug commands

diff --git a/Engine/Source/Editor/WorldBrowser/Private/LevelFolders.h b/Engine/Source/Editor/WorldBrowser/Private/LevelFolders.h
--- a/Engine/Source/Editor/WorldBrowser/Private/LevelFolders.h
+++ b/Engine/Source/Editor/WorldBrowser/Private/LevelFolders.h
@@ -65,6 +65,10 @@ public:
 
 	//~ FGCObject
 	virtual void AddReferencedObjects(FReferenceCollector& Collector) override;
+	virtual FString GetReferencerName() const
+	{
+		return "FLevelFolders";
+	}
 
 public:
 
