Stored reimport weightmap paths per tile when importing a tiled landscape (based on PR 6755 by moppius)

diff --git a/Engine/Source/Editor/WorldBrowser/Private/Tiles/WorldTileModel.cpp b/Engine/Source/Editor/WorldBrowser/Private/Tiles/WorldTileModel.cpp
--- a/Engine/Source/Editor/WorldBrowser/Private/Tiles/WorldTileModel.cpp
+++ b/Engine/Source/Editor/WorldBrowser/Private/Tiles/WorldTileModel.cpp
@@ -1019,6 +1019,11 @@ ALandscapeProxy* FWorldTileModel::ImportLandscapeTile(const FLandscapeImportSett
 	LandscapeProxy->Import(	Settings.LandscapeGuid, 0, 0, Settings.SizeX - 1, Settings.SizeY - 1, Settings.SectionsPerComponent, Settings.QuadsPerSection, HeightmapDataPerLayers, *Settings.HeightmapFilename,	
 							MaterialLayerDataPerLayer,	Settings.ImportLayerType);
 
+	for (const FLandscapeImportLayerInfo& ImportLayerInfo : Settings.ImportLayers)
+	{
+		LandscapeProxy->EditorLayerSettings.Add(FLandscapeEditorLayerSettings(ImportLayerInfo.LayerInfo, ImportLayerInfo.SourceFilePath));
+	}
+
 	return LandscapeProxy;
 }
 
