PR #6767: Clarified function and parameter names in FWorldTileModel and UWorldComposition (Contributed by [REDACTED_PERSON])

diff --git a/Engine/Source/Editor/WorldBrowser/Private/Tiles/WorldTileCollectionModel.cpp b/Engine/Source/Editor/WorldBrowser/Private/Tiles/WorldTileCollectionModel.cpp
--- a/Engine/Source/Editor/WorldBrowser/Private/Tiles/WorldTileCollectionModel.cpp
+++ b/Engine/Source/Editor/WorldBrowser/Private/Tiles/WorldTileCollectionModel.cpp
@@ -562,7 +562,7 @@ void FWorldTileCollectionModel::FillReimportTiledLandscapeSubMenu(FMenuBuilder&
 
 void FWorldTileCollectionModel::FillWeightmapsSubMenu(FMenuBuilder& InMenuBuilder) const
 {
-	// Add "All Weighmaps" menu entry
+	// Add "All Weightmaps" menu entry
 	InMenuBuilder.AddMenuEntry(
 			LOCTEXT("Menu_AllWeightmapsTitle", "All Weightmaps"), 
 			FText(), FSlateIcon(),
@@ -577,7 +577,7 @@ void FWorldTileCollectionModel::FillWeightmapsSubMenu(FMenuBuilder& InMenuBuilde
 	for (const auto& LevelModel : SelectedLevelsList)
 	{
 		auto TileModel = StaticCastSharedPtr<FWorldTileModel>(LevelModel);
-		if (TileModel->IsTiledLandscapeBased())
+		if (TileModel->IsLandscapeBased())
 		{
 			TArray<FName> Layers = ALandscapeProxy::GetLayersFromMaterial(TileModel->GetLandscape()->LandscapeMaterial);
 			for (FName LayerName : Layers)
@@ -1433,7 +1433,7 @@ bool FWorldTileCollectionModel::CanReimportTiledlandscape() const
 {
 	for (const auto& LevelModel : SelectedLevelsList)
 	{
-		if (LevelModel->IsEditable() && StaticCastSharedPtr<FWorldTileModel>(LevelModel)->IsTiledLandscapeBased())
+		if (LevelModel->IsEditable() && StaticCastSharedPtr<FWorldTileModel>(LevelModel)->CanReimportHeightmap())
 		{
 			return true;
 		}
@@ -1727,7 +1727,7 @@ void FWorldTileCollectionModel::ReimportTiledLandscape_Executed(FName TargetLaye
 	for (auto LevelModel : SelectedLevelsList)
 	{
 		TSharedPtr<FWorldTileModel> TileModel = StaticCastSharedPtr<FWorldTileModel>(LevelModel);
-		if (TileModel->IsEditable() && TileModel->IsTiledLandscapeBased())
+		if (TileModel->IsEditable() && TileModel->CanReimportHeightmap())
 		{
 			TargetLandscapeTiles.Add(TileModel);
 		}
@@ -1750,7 +1750,7 @@ void FWorldTileCollectionModel::ReimportTiledLandscape_Executed(FName TargetLaye
 	}
 
 	// Disable world origin tracking, so we can show, hide levels without offseting them
-	GetWorld()->WorldComposition->bTemporallyDisableOriginTracking = true;
+	GetWorld()->WorldComposition->bTemporarilyDisableOriginTracking = true;
 
 	// Reimport data for each selected landscape tile
 	for (auto TileModel : TargetLandscapeTiles)
@@ -1799,7 +1799,7 @@ void FWorldTileCollectionModel::ReimportTiledLandscape_Executed(FName TargetLaye
 	}
 
 	// Restore world origin tracking
-	GetWorld()->WorldComposition->bTemporallyDisableOriginTracking = false;
+	GetWorld()->WorldComposition->bTemporarilyDisableOriginTracking = false;
 
 	// Restore levels visibility
 	for (int32 LevelIdx = 0; LevelIdx < AllLevelsList.Num(); ++LevelIdx)
@@ -1974,7 +1974,7 @@ bool FWorldTileCollectionModel::GenerateLODLevels(FLevelModelList InLevelList, i
 		const bool bVisibleLevel = TileModel->IsVisible();
 		if (!bVisibleLevel)
 		{
-			GetWorld()->WorldComposition->bTemporallyDisableOriginTracking = true;
+			GetWorld()->WorldComposition->bTemporarilyDisableOriginTracking = true;
 			TileModel->SetVisible(true);
 		}
 
@@ -1997,7 +1997,7 @@ bool FWorldTileCollectionModel::GenerateLODLevels(FLevelModelList InLevelList, i
 
 		// This is current actors offset from their original position
 		FVector ActorsOffset = FVector(TileModel->GetAbsoluteLevelPosition() - GetWorld()->OriginLocation);
-		if (GetWorld()->WorldComposition->bTemporallyDisableOriginTracking)
+		if (GetWorld()->WorldComposition->bTemporarilyDisableOriginTracking)
 		{
 			ActorsOffset = FVector::ZeroVector;
 		}
@@ -2184,7 +2184,7 @@ bool FWorldTileCollectionModel::GenerateLODLevels(FLevelModelList InLevelList, i
 		if (!bVisibleLevel)
 		{
 			TileModel->SetVisible(false);
-			GetWorld()->WorldComposition->bTemporallyDisableOriginTracking = false;
+			GetWorld()->WorldComposition->bTemporarilyDisableOriginTracking = false;
 		}
 	
 		if (AssetsToSpawn.Num())
diff --git a/Engine/Source/Editor/WorldBrowser/Private/Tiles/WorldTileModel.cpp b/Engine/Source/Editor/WorldBrowser/Private/Tiles/WorldTileModel.cpp
--- a/Engine/Source/Editor/WorldBrowser/Private/Tiles/WorldTileModel.cpp
+++ b/Engine/Source/Editor/WorldBrowser/Private/Tiles/WorldTileModel.cpp
@@ -229,7 +229,7 @@ bool FWorldTileModel::ShouldBeVisible(FBox EditableArea) const
 	}
 
 	// When this hack is activated level should be visible regardless of current world origin
-	if (LevelCollectionModel.GetWorld()->WorldComposition->bTemporallyDisableOriginTracking)
+	if (LevelCollectionModel.GetWorld()->WorldComposition->bTemporarilyDisableOriginTracking)
 	{
 		return true;
 	}
@@ -335,7 +335,7 @@ bool FWorldTileModel::IsLandscapeBased() const
 	return Landscape.IsValid();
 }
 
-bool FWorldTileModel::IsTiledLandscapeBased() const
+bool FWorldTileModel::CanReimportHeightmap() const
 {
 	if (IsLandscapeBased() && !GetLandscape()->ReimportHeightmapFilePath.IsEmpty())
 	{
@@ -362,7 +362,7 @@ bool FWorldTileModel::IsTiledLandscapeBased() const
 	return false;
 }
 
-bool FWorldTileModel::IsLandscapeProxy() const
+bool FWorldTileModel::IsLandscapeStreamingProxy() const
 {
 	return (Landscape.IsValid() && Landscape.Get()->IsA(ALandscapeStreamingProxy::StaticClass()));
 }
diff --git a/Engine/Source/Editor/WorldBrowser/Private/Tiles/WorldTileModel.h b/Engine/Source/Editor/WorldBrowser/Private/Tiles/WorldTileModel.h
--- a/Engine/Source/Editor/WorldBrowser/Private/Tiles/WorldTileModel.h
+++ b/Engine/Source/Editor/WorldBrowser/Private/Tiles/WorldTileModel.h
@@ -120,11 +120,11 @@ public:
 	/** Whether this level landscape based or not */
 	bool IsLandscapeBased() const;
 
-	/** Whether this level based on tiled landscape or not */
-	bool IsTiledLandscapeBased() const;
+	/** Whether this level has a heightmap that can be reimported or not */
+	bool CanReimportHeightmap() const;
 		
-	/** Whether this level has ALandscapeProxy or not */
-	bool IsLandscapeProxy() const;
+	/** Whether this level has ALandscapeStreamingProxy or not */
+	bool IsLandscapeStreamingProxy() const;
 
 	/** @return The landscape actor in case this level is landscape based */
 	ALandscapeProxy* GetLandscape() const;
diff --git a/Engine/Source/Runtime/Engine/Classes/Engine/WorldComposition.h b/Engine/Source/Runtime/Engine/Classes/Engine/WorldComposition.h
--- a/Engine/Source/Runtime/Engine/Classes/Engine/WorldComposition.h
+++ b/Engine/Source/Runtime/Engine/Classes/Engine/WorldComposition.h
@@ -203,9 +203,9 @@ private:
 
 public:
 #if WITH_EDITOR
-	// Hack for a World Browser to be able to temporally show hidden levels 
-	// regardless of current world origin and without offsetting them temporally 
-	bool						bTemporallyDisableOriginTracking;
+	// Hack for a World Browser to be able to temporarily show hidden levels 
+	// regardless of current world origin and without offsetting them temporarily 
+	bool						bTemporarilyDisableOriginTracking;
 #endif //WITH_EDITOR
 
 private:
diff --git a/Engine/Source/Runtime/Engine/Private/WorldComposition.cpp b/Engine/Source/Runtime/Engine/Private/WorldComposition.cpp
--- a/Engine/Source/Runtime/Engine/Private/WorldComposition.cpp
+++ b/Engine/Source/Runtime/Engine/Private/WorldComposition.cpp
@@ -29,6 +29,9 @@ UWorldComposition::FWorldCompositionChangedEvent UWorldComposition::WorldComposi
 
 UWorldComposition::UWorldComposition(const FObjectInitializer& ObjectInitializer)
 	: Super(ObjectInitializer)
+#if WITH_EDITOR
+	, bTemporarilyDisableOriginTracking(false)
+#endif
 	, TilesStreamingTimeThreshold(1.0)
 	, bLoadAllTilesDuringCinematic(false)
 	, bRebaseOriginIn3DSpace(false)
@@ -871,7 +874,7 @@ bool UWorldComposition::CommitTileStreamingState(UWorld* PersistentWorld, int32
 void UWorldComposition::OnLevelAddedToWorld(ULevel* InLevel)
 {
 #if WITH_EDITOR
-	if (bTemporallyDisableOriginTracking)
+	if (bTemporarilyDisableOriginTracking)
 	{
 		return;
 	}
@@ -885,7 +888,7 @@ void UWorldComposition::OnLevelAddedToWorld(ULevel* InLevel)
 void UWorldComposition::OnLevelRemovedFromWorld(ULevel* InLevel)
 {
 #if WITH_EDITOR
-	if (bTemporallyDisableOriginTracking)
+	if (bTemporarilyDisableOriginTracking)
 	{
 		return;
 	}
