PR #6768: Show a progress bar when reimporting tiled Landscape (Contributed by [REDACTED_PERSON])

diff --git a/Engine/Source/Editor/WorldBrowser/Private/Tiles/WorldTileCollectionModel.cpp b/Engine/Source/Editor/WorldBrowser/Private/Tiles/WorldTileCollectionModel.cpp
--- a/Engine/Source/Editor/WorldBrowser/Private/Tiles/WorldTileCollectionModel.cpp
+++ b/Engine/Source/Editor/WorldBrowser/Private/Tiles/WorldTileCollectionModel.cpp
@@ -4,6 +4,7 @@
 #include "Components/PrimitiveComponent.h"
 #include "Misc/CoreDelegates.h"
 #include "Misc/MessageDialog.h"
+#include "Misc/ScopedSlowTask.h"
 #include "Modules/ModuleManager.h"
 #include "Widgets/SWindow.h"
 #include "Engine/MeshMerging.h"
@@ -1738,6 +1739,9 @@ void FWorldTileCollectionModel::ReimportTiledLandscape_Executed(FName TargetLaye
 		return;
 	}
 
+	FScopedSlowTask Progress(TargetLandscapeTiles.Num(), LOCTEXT("LandscapeImportProgress", "Reimporting landscape tiles..."));
+	Progress.MakeDialog();
+
 	TArray<bool> AllLevelsVisibilityState;
 	// Hide all visible levels
 	for (auto LevelModel : AllLevelsList)
@@ -1755,6 +1759,8 @@ void FWorldTileCollectionModel::ReimportTiledLandscape_Executed(FName TargetLaye
 	// Reimport data for each selected landscape tile
 	for (auto TileModel : TargetLandscapeTiles)
 	{
+		Progress.EnterProgressFrame();
+
 		TileModel->SetVisible(true);
 
 		ALandscapeProxy* Landscape = TileModel->GetLandscape();
