PR #6788: Fixed typo (Unock -> Unlock) and spacing consistency (Contributed by [REDACTED_PERSON])

diff --git a/Engine/Source/Editor/WorldBrowser/Private/LevelCollectionCommands.h b/Engine/Source/Editor/WorldBrowser/Private/LevelCollectionCommands.h
--- a/Engine/Source/Editor/WorldBrowser/Private/LevelCollectionCommands.h
+++ b/Engine/Source/Editor/WorldBrowser/Private/LevelCollectionCommands.h
@@ -41,7 +41,7 @@ public:
 		UI_COMMAND( World_MakeLevelCurrent, "Make Current", "Make this Level the Current Level", EUserInterfaceActionType::Button, FInputChord( EKeys::Enter ) );
 		UI_COMMAND( World_FindInContentBrowser, "Find in Content Browser", "Find the selected levels in the Content Browser", EUserInterfaceActionType::Button, FInputChord() );
 		UI_COMMAND( World_LoadLevel, "Load", "Load selected level into world", EUserInterfaceActionType::Button, FInputChord() );
-		UI_COMMAND(World_UnloadLevel, "Unload", "Unload selected level from world", EUserInterfaceActionType::Button, FInputChord(EKeys::Platform_Delete));
+		UI_COMMAND( World_UnloadLevel, "Unload", "Unload selected level from world", EUserInterfaceActionType::Button, FInputChord(EKeys::Platform_Delete) );
 		UI_COMMAND( World_SaveSelectedLevels, "Save", "Saves selected levels", EUserInterfaceActionType::Button, FInputChord() );
 		UI_COMMAND( World_SaveSelectedLevelAs, "Save As...", "Save the selected level as...", EUserInterfaceActionType::Button, FInputChord() );
 		UI_COMMAND( World_MigrateSelectedLevels, "Migrate...", "Copies the selected levels and all their dependencies to a different game", EUserInterfaceActionType::Button, FInputChord() );
@@ -49,7 +49,7 @@ public:
 		UI_COMMAND( World_CreateNewLevel, "Create New...", "Creates a new level, or choose a level template to start from.", EUserInterfaceActionType::Button, FInputChord() );
 		UI_COMMAND( World_AddExistingLevel, "Add Existing...", "Adds an existing level", EUserInterfaceActionType::Button, FInputChord() );
 		UI_COMMAND( World_AddSelectedActorsToNewLevel, "Create New with Selected Actors...", "Adds the actors currently selected in the active viewport to a new Level", EUserInterfaceActionType::Button, FInputChord() );
-		UI_COMMAND(World_RemoveSelectedLevels, "Remove Selected", "Removes selected levels from the base streaming level", EUserInterfaceActionType::Button, FInputChord() );
+		UI_COMMAND( World_RemoveSelectedLevels, "Remove Selected", "Removes selected levels from the base streaming level", EUserInterfaceActionType::Button, FInputChord() );
 
 		UI_COMMAND( MoveWorldOrigin, "Move World Origin to Level Position", "Moves world origin to level position", EUserInterfaceActionType::Button, FInputChord() );
 		UI_COMMAND( ResetWorldOrigin, "Reset World Origin", "Moves world origin to zero", EUserInterfaceActionType::Button, FInputChord() );
@@ -101,9 +101,9 @@ public:
 
 		//lock
 		UI_COMMAND( World_LockSelectedLevels, "Lock Selected", "Locks selected levels", EUserInterfaceActionType::Button, FInputChord() );
-		UI_COMMAND( World_UnockSelectedLevels, "Unlock Selected", "Unlocks selected levels", EUserInterfaceActionType::Button, FInputChord() );
-		UI_COMMAND(World_LockOnlySelectedLevels, "Lock Only Selected", "Toggles the selected levels to a locked state; toggles all other levels to an unlocked state", EUserInterfaceActionType::Button, FInputChord());
-		UI_COMMAND(World_LockAllButSelectedLevels, "Lock All But Selected", "Toggles the selected levels to an unlocked state; toggles all other levels to a locked state", EUserInterfaceActionType::Button, FInputChord());
+		UI_COMMAND( World_UnlockSelectedLevels, "Unlock Selected", "Unlocks selected levels", EUserInterfaceActionType::Button, FInputChord() );
+		UI_COMMAND( World_LockOnlySelectedLevels, "Lock Only Selected", "Toggles the selected levels to a locked state; toggles all other levels to an unlocked state", EUserInterfaceActionType::Button, FInputChord() );
+		UI_COMMAND( World_LockAllButSelectedLevels, "Lock All But Selected", "Toggles the selected levels to an unlocked state; toggles all other levels to a locked state", EUserInterfaceActionType::Button, FInputChord() );
 		UI_COMMAND( World_LockAllLevels, "Lock All", "Locks all levels", EUserInterfaceActionType::Button, FInputChord() );
 		UI_COMMAND( World_UnockAllLevels, "Unlock All", "Unlocks all levels", EUserInterfaceActionType::Button, FInputChord() );
 		UI_COMMAND( World_LockReadOnlyLevels, "Lock Read-Only Levels", "Locks all read-only levels", EUserInterfaceActionType::Button, FInputChord() );
@@ -281,7 +281,7 @@ public:
 	TSharedPtr< FUICommandInfo > World_LockSelectedLevels;
 
 	/** Unlocks selected levels */
-	TSharedPtr< FUICommandInfo > World_UnockSelectedLevels;
+	TSharedPtr< FUICommandInfo > World_UnlockSelectedLevels;
 
 	/** Locks all levels */
 	TSharedPtr< FUICommandInfo > World_LockAllLevels;
diff --git a/Engine/Source/Editor/WorldBrowser/Private/LevelCollectionModel.cpp b/Engine/Source/Editor/WorldBrowser/Private/LevelCollectionModel.cpp
--- a/Engine/Source/Editor/WorldBrowser/Private/LevelCollectionModel.cpp
+++ b/Engine/Source/Editor/WorldBrowser/Private/LevelCollectionModel.cpp
@@ -176,7 +176,7 @@ void FLevelCollectionModel::BindCommands()
 	ActionList.MapAction( Commands.World_LockSelectedLevels,
 		FExecuteAction::CreateSP( this, &FLevelCollectionModel::LockSelectedLevels_Executed  ) );
 	
-	ActionList.MapAction( Commands.World_UnockSelectedLevels,
+	ActionList.MapAction( Commands.World_UnlockSelectedLevels,
 		FExecuteAction::CreateSP( this, &FLevelCollectionModel::UnlockSelectedLevels_Executed  ) );
 
 	ActionList.MapAction(Commands.World_LockOnlySelectedLevels,
@@ -1810,9 +1810,9 @@ void FLevelCollectionModel::FillLockSubMenu(FMenuBuilder& InMenuBuilder)
 	const FLevelCollectionCommands& Commands = FLevelCollectionCommands::Get();
 
 	InMenuBuilder.AddMenuEntry( Commands.World_LockSelectedLevels );
-	InMenuBuilder.AddMenuEntry( Commands.World_UnockSelectedLevels );
-	InMenuBuilder.AddMenuEntry(Commands.World_LockOnlySelectedLevels);
-	InMenuBuilder.AddMenuEntry(Commands.World_LockAllButSelectedLevels);
+	InMenuBuilder.AddMenuEntry( Commands.World_UnlockSelectedLevels );
+	InMenuBuilder.AddMenuEntry( Commands.World_LockOnlySelectedLevels );
+	InMenuBuilder.AddMenuEntry( Commands.World_LockAllButSelectedLevels );
 	InMenuBuilder.AddMenuEntry( Commands.World_LockAllLevels );
 	InMenuBuilder.AddMenuEntry( Commands.World_UnockAllLevels );
 
@@ -1833,7 +1833,7 @@ void FLevelCollectionModel::FillVisibilitySubMenu(FMenuBuilder& InMenuBuilder)
 	InMenuBuilder.AddMenuEntry( Commands.World_ShowSelectedLevels );
 	InMenuBuilder.AddMenuEntry( Commands.World_HideSelectedLevels );
 	InMenuBuilder.AddMenuEntry( Commands.World_ShowOnlySelectedLevels );
-	InMenuBuilder.AddMenuEntry(Commands.World_ShowAllButSelectedLevels);
+	InMenuBuilder.AddMenuEntry( Commands.World_ShowAllButSelectedLevels );
 	InMenuBuilder.AddMenuEntry( Commands.World_ShowAllLevels );
 	InMenuBuilder.AddMenuEntry( Commands.World_HideAllLevels );
 }
