PR #6793: Fixed typo: ImportWidnow->ImportWindow (Contributed by [REDACTED_PERSON])

diff --git a/Engine/Source/Editor/WorldBrowser/Private/Tiles/WorldTileCollectionModel.cpp b/Engine/Source/Editor/WorldBrowser/Private/Tiles/WorldTileCollectionModel.cpp
--- a/Engine/Source/Editor/WorldBrowser/Private/Tiles/WorldTileCollectionModel.cpp
+++ b/Engine/Source/Editor/WorldBrowser/Private/Tiles/WorldTileCollectionModel.cpp
@@ -1586,18 +1586,18 @@ static void SetupLandscapeImportLayers(const FTiledLandscapeImportSettings& InIm
 void FWorldTileCollectionModel::ImportTiledLandscape_Executed()
 {
 	/** Create the window to host widget */
-	TSharedRef<SWindow> ImportWidnow = SNew(SWindow)
+	TSharedRef<SWindow> ImportWindow = SNew(SWindow)
 											.Title(LOCTEXT("TiledLandcapeImport_DialogTitle", "Import Tiled Landscape"))
 											.SizingRule( ESizingRule::Autosized )
 											.SupportsMinimize(false) 
 											.SupportsMaximize(false);
 
 	/** Set the content of the window */
-	TSharedRef<STiledLandcapeImportDlg> ImportDialog = SNew(STiledLandcapeImportDlg, ImportWidnow);
-	ImportWidnow->SetContent(ImportDialog);
+	TSharedRef<STiledLandcapeImportDlg> ImportDialog = SNew(STiledLandcapeImportDlg, ImportWindow);
+	ImportWindow->SetContent(ImportDialog);
 
 	/** Show the dialog window as a modal window */
-	GEditor->EditorAddModalWindow(ImportWidnow);
+	GEditor->EditorAddModalWindow(ImportWindow);
 
 	if (ImportDialog->ShouldImport() && ImportDialog->GetImportSettings().HeightmapFileList.Num())
 	{
