Context menu action to create subfolder no longer moves hierarchy into the new subfolder

diff --git a/Engine/Source/Editor/WorldBrowser/Private/SWorldHierarchyImpl.cpp b/Engine/Source/Editor/WorldBrowser/Private/SWorldHierarchyImpl.cpp
--- a/Engine/Source/Editor/WorldBrowser/Private/SWorldHierarchyImpl.cpp
+++ b/Engine/Source/Editor/WorldBrowser/Private/SWorldHierarchyImpl.cpp
@@ -547,7 +547,7 @@ void SWorldHierarchyImpl::FillFoldersSubmenu(FMenuBuilder& MenuBuilder)
 		LOCTEXT("CreateNewFolder", "Create New Folder"),
 		LOCTEXT("CreateNewFolder_Tooltip", "Move the selection to a new folder"),
 		FSlateIcon(FEditorStyle::GetStyleSetName(), "WorldBrowser.NewFolderIcon"),
-		FExecuteAction::CreateSP(this, &SWorldHierarchyImpl::CreateFolder, RootItem, RootPath)
+		FExecuteAction::CreateSP(this, &SWorldHierarchyImpl::CreateFolder, RootItem, RootPath, /*bMoveSelected*/ true)
 	);
 
 	AddMoveToFolderOutliner(MenuBuilder, SelectedItems, RootItem.ToSharedRef());
@@ -1398,7 +1398,7 @@ EVisibility SWorldHierarchyImpl::GetEmptyLabelVisibility() const
 	return ( !bFoldersOnlyMode || RootTreeItems.Num() > 0 ) ? EVisibility::Collapsed : EVisibility::Visible;
 }
 
-void SWorldHierarchyImpl::CreateFolder(TSharedPtr<FLevelModel> InModel, FName ParentPath /* = NAME_None */)
+void SWorldHierarchyImpl::CreateFolder(TSharedPtr<FLevelModel> InModel, FName ParentPath /* = NAME_None */, const bool bMoveSelected)
 {
 	if (FLevelFolders::IsAvailable())
 	{
@@ -1445,7 +1445,15 @@ void SWorldHierarchyImpl::CreateFolder(TSharedPtr<FLevelModel> InModel, FName Pa
 
 		NewFolderName = LevelFolders.GetDefaultFolderName(PersistentLevelModel.ToSharedRef(), NewFolderName);
 
-		MoveItemsTo(PersistentLevelModel, NewFolderName);
+		if (bMoveSelected)
+		{
+			MoveItemsTo(PersistentLevelModel, NewFolderName);
+		}
+		else if (FLevelFolders::IsAvailable())
+		{
+			LevelFolders.CreateFolder(PersistentLevelModel.ToSharedRef(), NewFolderName);
+			NewItemActions.Add(NewFolderName, WorldHierarchy::ENewItemAction::Select | WorldHierarchy::ENewItemAction::Rename);
+		}
 	}
 }
 
diff --git a/Engine/Source/Editor/WorldBrowser/Private/SWorldHierarchyImpl.h b/Engine/Source/Editor/WorldBrowser/Private/SWorldHierarchyImpl.h
--- a/Engine/Source/Editor/WorldBrowser/Private/SWorldHierarchyImpl.h
+++ b/Engine/Source/Editor/WorldBrowser/Private/SWorldHierarchyImpl.h
@@ -79,7 +79,7 @@ public:
 	virtual void Tick( const FGeometry& AllotedGeometry, const double InCurrentTime, const float InDeltaTime ) override;
 
 	/** Creates a new folder for the hierarchy. If ParentPath is defined, the folder will be created relative to that path. */
-	void CreateFolder(TSharedPtr<FLevelModel> InModel, FName ParentPath = NAME_None);
+	void CreateFolder(TSharedPtr<FLevelModel> InModel, FName ParentPath = NAME_None, const bool bMoveSelected = true);
 
 	/** Moves the current selection to the specified path */
 	void MoveItemsTo(TSharedPtr<FLevelModel> InModel, FName Path);
diff --git a/Engine/Source/Editor/WorldBrowser/Private/WorldTreeItemTypes.cpp b/Engine/Source/Editor/WorldBrowser/Private/WorldTreeItemTypes.cpp
--- a/Engine/Source/Editor/WorldBrowser/Private/WorldTreeItemTypes.cpp
+++ b/Engine/Source/Editor/WorldBrowser/Private/WorldTreeItemTypes.cpp
@@ -484,7 +484,7 @@ namespace WorldHierarchy
 			const FSlateIcon NewFolderIcon(FEditorStyle::GetStyleSetName(), "WorldBrowser.NewFolderIcon");
 
 			FName RootPath = LevelModel.Pin()->GetFolderPath();
-			auto NewFolderAction = FExecuteAction::CreateSP(const_cast<SWorldHierarchyImpl*>(&Hierarchy), &SWorldHierarchyImpl::CreateFolder, LevelModel.Pin(), RootPath);
+			auto NewFolderAction = FExecuteAction::CreateSP(const_cast<SWorldHierarchyImpl*>(&Hierarchy), &SWorldHierarchyImpl::CreateFolder, LevelModel.Pin(), RootPath, /*bMoveSelection*/ false);
 
 			FToolMenuSection& Section = Menu->AddSection("Section");
 			Section.AddMenuEntry("CreateFolder", LOCTEXT("CreateFolder", "Create Folder"), FText(), NewFolderIcon, FUIAction(NewFolderAction));
@@ -913,7 +913,7 @@ namespace WorldHierarchy
 		TArray<FWorldTreeItemPtr> Folders;
 		Folders.Add(AsShared());
 
-		auto NewFolderAction = FExecuteAction::CreateSP(const_cast<SWorldHierarchyImpl*>(&Hierarchy), &SWorldHierarchyImpl::CreateFolder, RootLevel, Path);
+		auto NewFolderAction = FExecuteAction::CreateSP(const_cast<SWorldHierarchyImpl*>(&Hierarchy), &SWorldHierarchyImpl::CreateFolder, RootLevel, Path, /*bMoveSelected*/ false);
 		auto RenameFolderAction = FExecuteAction::CreateSP(const_cast<SWorldHierarchyImpl*>(&Hierarchy), &SWorldHierarchyImpl::InitiateRename, AsShared());
 		auto DeleteFolderAction = FExecuteAction::CreateSP(const_cast<SWorldHierarchyImpl*>(&Hierarchy), &SWorldHierarchyImpl::DeleteFolders, Folders, /*bTransactional*/ true);
 
