When destroying the old LODWorld, rename it so we can reuse names for newly created actors.

diff --git a/Engine/Source/Editor/WorldBrowser/Private/Tiles/WorldTileCollectionModel.cpp b/Engine/Source/Editor/WorldBrowser/Private/Tiles/WorldTileCollectionModel.cpp
--- a/Engine/Source/Editor/WorldBrowser/Private/Tiles/WorldTileCollectionModel.cpp
+++ b/Engine/Source/Editor/WorldBrowser/Private/Tiles/WorldTileCollectionModel.cpp
@@ -2226,6 +2226,7 @@ bool FWorldTileCollectionModel::GenerateLODLevels(FLevelModelList InLevelList, i
 			{
 				LODWorld->ClearFlags(RF_Public | RF_Standalone);
 				LODWorld->DestroyWorld(false);
+				LODWorld->Rename(nullptr, GetTransientPackage());
 			}
 			
 			// Create a new world
