World Composition layer creation menu should create the layer on Enter keypress

diff --git a/Engine/Source/Editor/WorldBrowser/Private/Tiles/SWorldLayers.cpp b/Engine/Source/Editor/WorldBrowser/Private/Tiles/SWorldLayers.cpp
--- a/Engine/Source/Editor/WorldBrowser/Private/Tiles/SWorldLayers.cpp
+++ b/Engine/Source/Editor/WorldBrowser/Private/Tiles/SWorldLayers.cpp
@@ -60,6 +60,8 @@ void SNewWorldLayerPopup::Construct(const FArguments& InArgs)
 					.Text(this, &SNewWorldLayerPopup::GetLayerName)
 					.SelectAllTextWhenFocused(true)
 					.OnTextChanged(this, &SNewWorldLayerPopup::SetLayerName)
+					.OnTextCommitted(this, &SNewWorldLayerPopup::OnNameCommitted)
+					.ClearKeyboardFocusOnCommit(false)
 				]
 
 			]
@@ -87,6 +89,7 @@ void SNewWorldLayerPopup::Construct(const FArguments& InArgs)
 					.MinValue(1)
 					.MaxValue(TNumericLimits<int32>::Max())
 					.OnValueChanged(this, &SNewWorldLayerPopup::SetStreamingDistance)
+					.OnValueCommitted(this, &SNewWorldLayerPopup::OnDistanceCommitted)
 					.LabelPadding(0)
 					.Label()
 					[
@@ -116,17 +119,39 @@ END_SLATE_FUNCTION_BUILD_OPTIMIZATION
 
 FReply SNewWorldLayerPopup::OnClickedCreate()
 {
-	if (OnCreateLayer.IsBound())
+	return TryCreateLayer();
+}
+
+void SNewWorldLayerPopup::OnNameCommitted(const FText& InText, ETextCommit::Type CommitType)
+{
+	if (CommitType == ETextCommit::OnEnter)
+	{
+		TryCreateLayer();
+	}
+}
+
+void SNewWorldLayerPopup::OnDistanceCommitted(int32 InValue, ETextCommit::Type CommitType)
+{
+	if (CommitType == ETextCommit::OnEnter)
+	{
+		TryCreateLayer();
+	}
+}
+
+FReply SNewWorldLayerPopup::TryCreateLayer()
+{
+	if (CanCreateLayer())
 	{
 		return OnCreateLayer.Execute(LayerData);
 	}
 	
+	// Return an unhandled reply if the layer should not be created
 	return FReply::Unhandled();
 }
 
 bool SNewWorldLayerPopup::CanCreateLayer() const
 {
-	return LayerData.Name.Len() > 0 && !ExistingLayerNames.Contains(LayerData.Name);
+	return LayerData.Name.Len() > 0 && !ExistingLayerNames.Contains(LayerData.Name) && OnCreateLayer.IsBound();
 }
 
 /** A class for check boxes in the layer list. 
diff --git a/Engine/Source/Editor/WorldBrowser/Private/Tiles/SWorldLayers.h b/Engine/Source/Editor/WorldBrowser/Private/Tiles/SWorldLayers.h
--- a/Engine/Source/Editor/WorldBrowser/Private/Tiles/SWorldLayers.h
+++ b/Engine/Source/Editor/WorldBrowser/Private/Tiles/SWorldLayers.h
@@ -85,6 +85,13 @@ public:
 	
 private:
 	FReply OnClickedCreate();
+
+	void OnNameCommitted(const FText& InText, ETextCommit::Type CommitType);
+	void OnDistanceCommitted(int32 InValue, ETextCommit::Type CommitType);
+
+	/** Try to create a new streaming layer */
+	FReply TryCreateLayer();
+
 	bool CanCreateLayer() const;
 
 	void SetLayerName(const FText& InText)
