Fix missing visibility icons in level browser

diff --git a/Engine/Source/Editor/WorldBrowser/Private/SWorldHierarchyItem.cpp b/Engine/Source/Editor/WorldBrowser/Private/SWorldHierarchyItem.cpp
--- a/Engine/Source/Editor/WorldBrowser/Private/SWorldHierarchyItem.cpp
+++ b/Engine/Source/Editor/WorldBrowser/Private/SWorldHierarchyItem.cpp
@@ -135,6 +135,7 @@ TSharedRef< SWidget > SWorldHierarchyItem::GenerateWidgetForColumn( const FName&
 				[
 					SNew(SImage)
 					.Image(this, &SWorldHierarchyItem::GetLightingScenarioBrush)
+					.ColorAndOpacity(FSlateColor::UseForeground())
 				]
 			;
 		}
@@ -153,6 +154,7 @@ TSharedRef< SWidget > SWorldHierarchyItem::GenerateWidgetForColumn( const FName&
 				[
 					SNew(SImage)
 					.Image(this, &SWorldHierarchyItem::GetLevelLockBrush)
+					.ColorAndOpacity(FSlateColor::UseForeground())
 				]
 			;
 		}
@@ -161,7 +163,7 @@ TSharedRef< SWidget > SWorldHierarchyItem::GenerateWidgetForColumn( const FName&
 			TableRowContent =
 				SAssignNew(VisibilityButton, SButton)
 				.ContentPadding(0)
-				.ButtonStyle(FEditorStyle::Get(), "ToggleButton")
+				.ButtonStyle(FEditorStyle::Get(), "SimpleButton")
 				.IsEnabled(this, &SWorldHierarchyItem::IsVisibilityEnabled)
 				.OnClicked(this, &SWorldHierarchyItem::OnToggleVisibility)
 				.ToolTipText(this, &SWorldHierarchyItem::GetVisibilityToolTip)
@@ -171,6 +173,7 @@ TSharedRef< SWidget > SWorldHierarchyItem::GenerateWidgetForColumn( const FName&
 				[
 					SNew(SImage)
 					.Image(this, &SWorldHierarchyItem::GetLevelVisibilityBrush)
+					.ColorAndOpacity(FSlateColor::UseForeground())
 				]
 			;
 		}
@@ -179,7 +182,7 @@ TSharedRef< SWidget > SWorldHierarchyItem::GenerateWidgetForColumn( const FName&
 			TableRowContent =
 				SAssignNew(ColorButton, SButton)
 				.ContentPadding(0)
-				.ButtonStyle(FEditorStyle::Get(), "ToggleButton")
+				.ButtonStyle(FEditorStyle::Get(), "SimpleButton")
 				.IsEnabled(true)
 				.OnClicked(this, &SWorldHierarchyItem::OnChangeColor)
 				.ToolTipText(LOCTEXT("LevelColorButtonToolTip", "Change Level Color"))
@@ -190,7 +193,8 @@ TSharedRef< SWidget > SWorldHierarchyItem::GenerateWidgetForColumn( const FName&
 				[
 					SNew(SImage)
 					.ColorAndOpacity(this, &SWorldHierarchyItem::GetDrawColor)
-				.Image(this, &SWorldHierarchyItem::GetLevelColorBrush)
+					.Image(this, &SWorldHierarchyItem::GetLevelColorBrush)
+					.ColorAndOpacity(FSlateColor::UseForeground())
 				]
 			;
 		}
