Integrate 14115533 from RES: Expose External Actors Conversion to the Level tab context menu.

diff --git a/Engine/Source/Editor/WorldBrowser/Private/LevelCollectionCommands.h b/Engine/Source/Editor/WorldBrowser/Private/LevelCollectionCommands.h
--- a/Engine/Source/Editor/WorldBrowser/Private/LevelCollectionCommands.h
+++ b/Engine/Source/Editor/WorldBrowser/Private/LevelCollectionCommands.h
@@ -87,6 +87,8 @@ public:
 		//actors
 		UI_COMMAND( AddsActors, "Select Actors", "Adds the Actors in the selected Levels from the viewport's existing selection", EUserInterfaceActionType::Button, FInputChord() );
 		UI_COMMAND( RemovesActors, "Deselect Actors", "Removes the Actors in the selected Levels from the viewport's existing selection", EUserInterfaceActionType::Button, FInputChord() );
+		UI_COMMAND (ConvertLevelToExternalActors, "Convert Level to External Actors", "Enable selected levels to use external actors and convert their actors to external packaging", EUserInterfaceActionType::Button, FInputChord() );
+		UI_COMMAND( ConvertLevelToInternalActors, "Convert Level to Internal Actors", "Disable selected levels to use external actors and convert their actors to internal packaging", EUserInterfaceActionType::Button, FInputChord() );
 		UI_COMMAND( MoveActorsToSelected, "Move Selected Actors to Level", "Moves the selected actors to this level", EUserInterfaceActionType::Button, FInputChord() );
 		UI_COMMAND( MoveFoliageToSelected, "Move Selected Foliage to Level", "Moves the selected foliage instances to this level. Keeps cross-level references to original bases", EUserInterfaceActionType::Button, FInputChord() );
 		UI_COMMAND( SelectStreamingVolumes, "Select Streaming Volumes", "Selects the streaming volumes associated with the selected levels", EUserInterfaceActionType::Button, FInputChord() );
@@ -248,6 +250,12 @@ public:
 	/** Deselects the actors in the selected Levels */
 	TSharedPtr< FUICommandInfo > RemovesActors;
 	
+	/** Activate `Use External Actors` on the level and convert actors to external packaging. */
+	TSharedPtr< FUICommandInfo > ConvertLevelToExternalActors;
+
+	/** Disable `Use External Actors` on the level and convert actors to internal packaging. */
+	TSharedPtr< FUICommandInfo > ConvertLevelToInternalActors;
+
 	/** Moves the selected actors to the selected level */
 	TSharedPtr< FUICommandInfo > MoveActorsToSelected;
 
@@ -257,7 +265,6 @@ public:
 	/** Selects the streaming volumes associated with the selected levels */
 	TSharedPtr< FUICommandInfo > SelectStreamingVolumes;
 	
-	
 	/** Makes selected Levels visible */
 	TSharedPtr< FUICommandInfo > World_ShowSelectedLevels;
 
diff --git a/Engine/Source/Editor/WorldBrowser/Private/LevelCollectionModel.cpp b/Engine/Source/Editor/WorldBrowser/Private/LevelCollectionModel.cpp
--- a/Engine/Source/Editor/WorldBrowser/Private/LevelCollectionModel.cpp
+++ b/Engine/Source/Editor/WorldBrowser/Private/LevelCollectionModel.cpp
@@ -151,6 +151,14 @@ void FLevelCollectionModel::BindCommands()
 		FExecuteAction::CreateSP(this, &FLevelCollectionModel::DeselectActors_Executed),
 		FCanExecuteAction::CreateSP(this, &FLevelCollectionModel::AreAnySelectedLevelsEditable));
 
+	ActionList.MapAction(Commands.ConvertLevelToExternalActors,
+		FExecuteAction::CreateSP(this, &FLevelCollectionModel::ConvertLevelToExternalActors_Executed, true),
+		FCanExecuteAction::CreateSP(this, &FLevelCollectionModel::CanConvertAnyLevelToExternalActors, true));
+
+	ActionList.MapAction(Commands.ConvertLevelToInternalActors,
+		FExecuteAction::CreateSP(this, &FLevelCollectionModel::ConvertLevelToExternalActors_Executed, false),
+		FCanExecuteAction::CreateSP(this, &FLevelCollectionModel::CanConvertAnyLevelToExternalActors, false));
+
 	//visibility
 	ActionList.MapAction( Commands.World_ShowSelectedLevels,
 		FExecuteAction::CreateSP( this, &FLevelCollectionModel::ShowSelectedLevels_Executed  ),
@@ -1078,6 +1086,18 @@ bool FLevelCollectionModel::AreActorsSelected() const
 	return GEditor->GetSelectedActorCount() > 0;
 }
 
+bool FLevelCollectionModel::CanConvertAnyLevelToExternalActors(bool bExternal) const
+{
+	for (const TSharedPtr<FLevelModel>& LevelModel : SelectedLevelsList)
+	{
+		if (!LevelModel->CanConvertLevelToExternalActors(bExternal))
+		{
+			return false;
+		}
+	}
+	return true;
+}
+
 bool FLevelCollectionModel::GetDisplayPathsState() const
 {
 	return bDisplayPaths;
@@ -1821,6 +1841,15 @@ void FLevelCollectionModel::DeselectActors_Executed()
 	}
 }
 
+void FLevelCollectionModel::ConvertLevelToExternalActors_Executed(bool bExternal)
+{
+	FScopedTransaction Transaction(LOCTEXT("WorldUseExternalActors", "Change World Use External Actors"));
+	for (const TSharedPtr<FLevelModel>& LevelModel : SelectedLevelsList)
+	{
+		LevelModel->ConvertLevelToExternalActors(bExternal);
+	}
+}
+
 void FLevelCollectionModel::ExpandSelectedItems_Executed()
 {
 	struct FExpandLevelVisitor : public FLevelModelVisitor
diff --git a/Engine/Source/Editor/WorldBrowser/Private/LevelCollectionModel.h b/Engine/Source/Editor/WorldBrowser/Private/LevelCollectionModel.h
--- a/Engine/Source/Editor/WorldBrowser/Private/LevelCollectionModel.h
+++ b/Engine/Source/Editor/WorldBrowser/Private/LevelCollectionModel.h
@@ -251,6 +251,9 @@ public:
 	/** @return	whether at least one actor is selected */
 	bool AreActorsSelected() const;
 
+	/** @return whether any of the currently selected levels can be converted to the specified actor bExternal packaging */
+	bool CanConvertAnyLevelToExternalActors(bool bExternal) const;
+
 	/** @return whether moving the selected actors to the selected level is a valid action */
 	bool IsValidMoveActorsToLevel() const;
 
@@ -415,6 +418,9 @@ protected:
 	/** Removes the Actors in the selected Levels from the viewport's existing selection */
 	void DeselectActors_Executed();
 
+	/** Set level `Use External Actors` to bExternal  */
+	void ConvertLevelToExternalActors_Executed(bool bExternal);
+
 	/** Toggles selected levels to a visible state in the viewports */
 	void ShowSelectedLevels_Executed();
 
diff --git a/Engine/Source/Editor/WorldBrowser/Private/LevelModel.cpp b/Engine/Source/Editor/WorldBrowser/Private/LevelModel.cpp
--- a/Engine/Source/Editor/WorldBrowser/Private/LevelModel.cpp
+++ b/Engine/Source/Editor/WorldBrowser/Private/LevelModel.cpp
@@ -749,6 +749,38 @@ void FLevelModel::SelectActors(bool bSelect, bool bNotify, bool bSelectEvenIfHid
 	}
 }
 
+void FLevelModel::ConvertLevelToExternalActors(bool bUseExternal)
+{
+	ULevel* Level = GetLevelObject();
+	if (Level == nullptr || IsLocked())
+	{
+		return;
+	}
+	Level->Modify();
+	Level->SetUseExternalActors(bUseExternal);
+	Level->ConvertAllActorsToPackaging(Level->IsUsingExternalActors());
+}
+
+bool FLevelModel::CanConvertLevelToExternalActors(bool bToExternal)
+{
+	ULevel* Level = GetLevelObject();
+	if (Level == nullptr || IsLocked())
+	{
+		return false;
+	}
+
+	UPackage* LevelPackage = Level->GetOutermost();
+	if (Level->IsUsingExternalActors() == bToExternal
+		|| LevelPackage == GetTransientPackage()
+		|| LevelPackage->HasAnyFlags(RF_Transient)
+		|| Level->HasAnyFlags(RF_Transient)
+		|| !FPackageName::IsValidLongPackageName(LevelPackage->GetName()))
+	{
+		return false;
+	}
+	return true;
+}
+
 void FLevelModel::UpdateLevelActorsCount()
 {
 	LevelActorsCount = 0;
diff --git a/Engine/Source/Editor/WorldBrowser/Private/LevelModel.h b/Engine/Source/Editor/WorldBrowser/Private/LevelModel.h
--- a/Engine/Source/Editor/WorldBrowser/Private/LevelModel.h
+++ b/Engine/Source/Editor/WorldBrowser/Private/LevelModel.h
@@ -271,6 +271,19 @@ public:
 	void SelectActors(bool bSelect, bool bNotify, bool bSelectEvenIfHidden, 
 						const TSharedPtr<ActorFilter>& Filter = TSharedPtr<ActorFilter>(NULL));
 	
+	/**
+	 * Set the level to use external actors and convert its actors to external packaging or vice versa
+	 * @param bUseExternal Convert to external if `bUseExternal` is true, convert to internal otherwise
+	 */
+	void ConvertLevelToExternalActors(bool bUseExternal);
+
+	/**
+	 * Return if the level can be converted to the specified external packaging mode.
+	 * @param bToExternal The packaging mode we want to convert to. (i.e true: external, false: internal)
+	 * @return true if we can convert the level to the `bToExternal` packaging mode.
+	 */
+	bool CanConvertLevelToExternalActors(bool bToExternal);
+
 	/** Updates cached value of level actors count */
 	void UpdateLevelActorsCount();
 
diff --git a/Engine/Source/Editor/WorldBrowser/Private/StreamingLevels/StreamingLevelCollectionModel.cpp b/Engine/Source/Editor/WorldBrowser/Private/StreamingLevels/StreamingLevelCollectionModel.cpp
--- a/Engine/Source/Editor/WorldBrowser/Private/StreamingLevels/StreamingLevelCollectionModel.cpp
+++ b/Engine/Source/Editor/WorldBrowser/Private/StreamingLevels/StreamingLevelCollectionModel.cpp
@@ -349,6 +349,9 @@ void FStreamingLevelCollectionModel::BuildHierarchyMenu(FMenuBuilder& InMenuBuil
 			InMenuBuilder.AddMenuEntry( Commands.MoveFoliageToSelected );
 		}
 
+		InMenuBuilder.AddMenuEntry(Commands.ConvertLevelToExternalActors);
+		InMenuBuilder.AddMenuEntry(Commands.ConvertLevelToInternalActors);
+
 		if (AreAnyLevelsSelected() && !(IsOneLevelSelected() && SelectedLevelsList[0]->IsPersistent()))
 		{
 			InMenuBuilder.AddMenuEntry( Commands.SelectStreamingVolumes );
