WorldPartition: Remove Summon from Levels Window and Add it to the Window menu

diff --git a/Engine/Source/Editor/EditorStyle/Private/SlateEditorStyle.cpp b/Engine/Source/Editor/EditorStyle/Private/SlateEditorStyle.cpp
--- a/Engine/Source/Editor/EditorStyle/Private/SlateEditorStyle.cpp
+++ b/Engine/Source/Editor/EditorStyle/Private/SlateEditorStyle.cpp
@@ -2300,7 +2300,6 @@ void FSlateEditorStyle::FStyle::SetupWorldBrowserStyles()
 		Set("WorldBrowser.HierarchyButtonBrush", new IMAGE_BRUSH("Icons/icon_levels_hierarchybutton_16x", Icon16x16));
 		Set("WorldBrowser.DetailsButtonBrush", new IMAGE_BRUSH("Icons/icon_levels_detailsbutton_40x", Icon16x16));
 		Set("WorldBrowser.CompositionButtonBrush", new IMAGE_BRUSH("Icons/icon_levels_compositionbutton_16x", Icon16x16));
-		Set("WorldBrowser.PartitionEditorButtonBrush", new IMAGE_BRUSH( "/Icons/icon_levels_partitionbutton_16x", Icon16x16 ) );
 
 		Set("WorldBrowser.FolderClosed", new IMAGE_BRUSH("Icons/FolderClosed", Icon16x16));
 		Set("WorldBrowser.FolderOpen", new IMAGE_BRUSH("Icons/FolderOpen", Icon16x16));
diff --git a/Engine/Source/Editor/EditorStyle/Private/StarshipStyle.cpp b/Engine/Source/Editor/EditorStyle/Private/StarshipStyle.cpp
--- a/Engine/Source/Editor/EditorStyle/Private/StarshipStyle.cpp
+++ b/Engine/Source/Editor/EditorStyle/Private/StarshipStyle.cpp
@@ -2025,7 +2025,6 @@ void FStarshipEditorStyle::FStyle::SetupWorldBrowserStyles()
 		Set("WorldBrowser.HierarchyButtonBrush", new IMAGE_BRUSH("Icons/icon_levels_hierarchybutton_16x", Icon16x16));
 		Set("WorldBrowser.DetailsButtonBrush", new IMAGE_BRUSH("Icons/icon_levels_detailsbutton_40x", Icon16x16));
 		Set("WorldBrowser.CompositionButtonBrush", new IMAGE_BRUSH("Icons/icon_levels_compositionbutton_16x", Icon16x16));
-		Set("WorldBrowser.PartitionEditorButtonBrush", new IMAGE_BRUSH( "/Icons/icon_levels_partitionbutton_16x", Icon16x16 ) );
 
 		Set("WorldBrowser.FolderClosed", new IMAGE_BRUSH("Icons/FolderClosed", Icon16x16));
 		Set("WorldBrowser.FolderOpen", new IMAGE_BRUSH("Icons/FolderOpen", Icon16x16));
diff --git a/Engine/Source/Editor/LevelEditor/Private/LevelEditor.cpp b/Engine/Source/Editor/LevelEditor/Private/LevelEditor.cpp
--- a/Engine/Source/Editor/LevelEditor/Private/LevelEditor.cpp
+++ b/Engine/Source/Editor/LevelEditor/Private/LevelEditor.cpp
@@ -376,12 +376,6 @@ void FLevelEditorModule::SummonWorldBrowserComposition()
 	LevelEditorInstance->TryInvokeTab(LevelEditorTabIds::WorldBrowserComposition);
 }
 
-void FLevelEditorModule::SummonWorldBrowserPartitionEditor()
-{
-	TSharedPtr<SLevelEditor> LevelEditorInstance = LevelEditorInstancePtr.Pin();
-	LevelEditorInstance->TryInvokeTab(LevelEditorTabIds::WorldBrowserPartitionEditor);
-}
-
 // @todo remove when world-centric mode is added
 void FLevelEditorModule::AttachSequencer( TSharedPtr<SWidget> SequencerWidget, TSharedPtr<IAssetEditorInstance> SequencerAssetEditor )
 {
diff --git a/Engine/Source/Editor/LevelEditor/Private/SLevelEditor.cpp b/Engine/Source/Editor/LevelEditor/Private/SLevelEditor.cpp
--- a/Engine/Source/Editor/LevelEditor/Private/SLevelEditor.cpp
+++ b/Engine/Source/Editor/LevelEditor/Private/SLevelEditor.cpp
@@ -1243,9 +1243,9 @@ TSharedRef<SWidget> SLevelEditor::RestoreContentArea( const TSharedRef<SDockTab>
 				.SetGroup( WorkspaceMenu::GetMenuStructure().GetLevelEditorCategory() )
 				.SetIcon( FSlateIcon(FEditorStyle::GetStyleSetName(), "LevelEditor.Tabs.WorldBrowserComposition") );
 
-			LevelEditorTabManager->RegisterTabSpawner( WorldPartitionEditorTab, FOnSpawnTab::CreateSP<SLevelEditor, FName, FString>(this, &SLevelEditor::SpawnLevelEditorTab, WorldPartitionEditorTab, FString()) )
-				.SetMenuType( ETabSpawnerMenuType::Hidden )
-				.SetDisplayName(NSLOCTEXT("GridEditorTabs", "WorldPartitionEditor", "World Grid Editor"))
+			LevelEditorTabManager->RegisterTabSpawner( WorldPartitionEditorTab, FOnSpawnTab::CreateSP<SLevelEditor, FName, FString>(this, &SLevelEditor::SpawnLevelEditorTab, LevelEditorTabIds::WorldBrowserPartitionEditor, FString()) )
+				.SetDisplayName(NSLOCTEXT("LevelEditorTabs", "WorldPartitionEditor", "World Partition"))
+				.SetTooltipText(NSLOCTEXT("LevelEditorTabs", "WorldPartitionEditorTooltipText", "Open World Partition Editor."))
 				.SetGroup( WorkspaceMenu::GetMenuStructure().GetLevelEditorCategory() )
 				.SetIcon( FSlateIcon(FEditorStyle::GetStyleSetName(), "LevelEditor.Tabs.WorldPartition") );
 		}
diff --git a/Engine/Source/Editor/LevelEditor/Public/LevelEditor.h b/Engine/Source/Editor/LevelEditor/Public/LevelEditor.h
--- a/Engine/Source/Editor/LevelEditor/Public/LevelEditor.h
+++ b/Engine/Source/Editor/LevelEditor/Public/LevelEditor.h
@@ -101,8 +101,7 @@ public:
 	virtual void SummonWorldBrowserHierarchy();
 	virtual void SummonWorldBrowserDetails();
 	virtual void SummonWorldBrowserComposition();
-	virtual void SummonWorldBrowserPartitionEditor();
-
+	
 	// @todo remove when world-centric mode is added
 	/**
 	 * Spawns a new sequencer tab if one doesn't exist already
diff --git a/Engine/Source/Editor/WorldBrowser/Private/SWorldHierarchy.cpp b/Engine/Source/Editor/WorldBrowser/Private/SWorldHierarchy.cpp
--- a/Engine/Source/Editor/WorldBrowser/Private/SWorldHierarchy.cpp
+++ b/Engine/Source/Editor/WorldBrowser/Private/SWorldHierarchy.cpp
@@ -153,26 +153,6 @@ void SWorldHierarchy::OnBrowseWorld(UWorld* InWorld)
 							.Image(this, &SWorldHierarchy::GetSummonCompositionBrush)
 						]
 					]
-
-					// Button to summon world partition editor tab
-					+SHorizontalBox::Slot()
-					.AutoWidth()
-					.VAlign(VAlign_Center)
-					.HAlign(HAlign_Left)
-					[
-						SNew(SButton)
-						.Visibility(this, &SWorldHierarchy::GetPartitionEditorButtonVisibility)
-						.ButtonStyle(FEditorStyle::Get(), "ToggleButton")
-						.OnClicked(this, &SWorldHierarchy::OnSummonPartitionEditor)
-						.ToolTipText(LOCTEXT("SummonPartitionEditorToolTipText", "Summons world partition editor"))
-						.HAlign(HAlign_Center)
-						.VAlign(VAlign_Center)
-						.Content()
-						[
-							SNew(SImage)
-							.Image(this, &SWorldHierarchy::GetSummonPartitionEditorBrush)
-						]
-					]
 				]
 			]
 			
@@ -221,23 +201,6 @@ FReply SWorldHierarchy::OnSummonComposition()
 	return FReply::Handled();
 }
 
-const FSlateBrush* SWorldHierarchy::GetSummonPartitionEditorBrush() const
-{
-	return FEditorStyle::GetBrush("WorldBrowser.PartitionEditorButtonBrush");
-}
-
-EVisibility SWorldHierarchy::GetPartitionEditorButtonVisibility() const
-{
-	return WorldModel->IsPartitionWorld() ? EVisibility::Visible : EVisibility::Collapsed;
-}
-
-FReply SWorldHierarchy::OnSummonPartitionEditor()
-{
-	FLevelEditorModule& LevelEditorModule = FModuleManager::LoadModuleChecked<FLevelEditorModule>( "LevelEditor" );
-	LevelEditorModule.SummonWorldBrowserPartitionEditor();
-	return FReply::Handled();
-}
-
 const FSlateBrush* SWorldHierarchy::GetSummonCompositionBrush() const
 {
 	return FEditorStyle::GetBrush("WorldBrowser.CompositionButtonBrush");
diff --git a/Engine/Source/Editor/WorldBrowser/Private/SWorldHierarchy.h b/Engine/Source/Editor/WorldBrowser/Private/SWorldHierarchy.h
--- a/Engine/Source/Editor/WorldBrowser/Private/SWorldHierarchy.h
+++ b/Engine/Source/Editor/WorldBrowser/Private/SWorldHierarchy.h
@@ -40,11 +40,6 @@ private:
 	FReply OnSummonComposition();
 	const FSlateBrush* GetSummonCompositionBrush() const;
 
-	/**  */
-	EVisibility GetPartitionEditorButtonVisibility() const;
-	FReply OnSummonPartitionEditor();
-	const FSlateBrush* GetSummonPartitionEditorBrush() const;
-
 	/**  */
 	TSharedRef<SWidget> GetFileButtonContent();
 
