Don't allow spawning the level editor when the world is partitioned.

diff --git a/Engine/Source/Editor/LevelEditor/Private/SLevelEditor.cpp b/Engine/Source/Editor/LevelEditor/Private/SLevelEditor.cpp
--- a/Engine/Source/Editor/LevelEditor/Private/SLevelEditor.cpp
+++ b/Engine/Source/Editor/LevelEditor/Private/SLevelEditor.cpp
@@ -979,6 +979,16 @@ bool SLevelEditor::HasAnyHostedEditorModeToolkit() const
 	return false;
 }
 
+bool SLevelEditor::CanSpawnWorldBrowser(const FSpawnTabArgs& Args) const
+{
+	return !UWorld::HasSubsystem<UWorldPartitionSubsystem>(GetWorld());
+}
+
+bool SLevelEditor::SLevelEditor::CanSpawnWorldPartitionBrowser(const FSpawnTabArgs& Args) const
+{
+	return UWorld::HasSubsystem<UWorldPartitionSubsystem>(GetWorld());
+}
+
 bool SLevelEditor::CanSpawnLayerBrowser(const FSpawnTabArgs& Args) const
 {
 	return !UWorld::HasSubsystem<UWorldPartitionSubsystem>(GetWorld());
@@ -1309,7 +1319,8 @@ TSharedRef<SWidget> SLevelEditor::RestoreContentArea( const TSharedRef<SDockTab>
 		}
 		
 		{
-			LevelEditorTabManager->RegisterTabSpawner(LevelEditorTabIds::WorldBrowserHierarchy, FOnSpawnTab::CreateSP<SLevelEditor, FName, FString>(this, &SLevelEditor::SpawnLevelEditorTab, LevelEditorTabIds::WorldBrowserHierarchy, FString()) )
+			LevelEditorTabManager->RegisterTabSpawner(LevelEditorTabIds::WorldBrowserHierarchy, FOnSpawnTab::CreateSP<SLevelEditor, FName, FString>(this, &SLevelEditor::SpawnLevelEditorTab, LevelEditorTabIds::WorldBrowserHierarchy, FString()),
+																								FCanSpawnTab::CreateSP(this, &SLevelEditor::CanSpawnWorldBrowser))
 				.SetDisplayName(NSLOCTEXT("LevelEditorTabs", "WorldBrowserHierarchy", "Levels"))
 				.SetTooltipText(NSLOCTEXT("LevelEditorTabs", "WorldBrowserHierarchyTooltipText", "Open the Levels tab. Use this to manage the levels in the current project."))
 				.SetGroup( WorkspaceMenu::GetMenuStructure().GetLevelEditorCategory() )
@@ -1327,7 +1338,8 @@ TSharedRef<SWidget> SLevelEditor::RestoreContentArea( const TSharedRef<SDockTab>
 				.SetGroup( WorkspaceMenu::GetMenuStructure().GetLevelEditorCategory() )
 				.SetIcon( FSlateIcon(FEditorStyle::GetStyleSetName(), "LevelEditor.Tabs.WorldBrowserComposition") );
 
-			LevelEditorTabManager->RegisterTabSpawner( WorldPartitionEditorTab, FOnSpawnTab::CreateSP<SLevelEditor, FName, FString>(this, &SLevelEditor::SpawnLevelEditorTab, LevelEditorTabIds::WorldBrowserPartitionEditor, FString()) )
+			LevelEditorTabManager->RegisterTabSpawner( WorldPartitionEditorTab, FOnSpawnTab::CreateSP<SLevelEditor, FName, FString>(this, &SLevelEditor::SpawnLevelEditorTab, LevelEditorTabIds::WorldBrowserPartitionEditor, FString()),
+																				FCanSpawnTab::CreateSP(this, &SLevelEditor::CanSpawnWorldPartitionBrowser))
 				.SetDisplayName(NSLOCTEXT("LevelEditorTabs", "WorldPartitionEditor", "World Partition"))
 				.SetTooltipText(NSLOCTEXT("LevelEditorTabs", "WorldPartitionEditorTooltipText", "Open World Partition Editor."))
 				.SetGroup( WorkspaceMenu::GetMenuStructure().GetLevelEditorCategory() )
diff --git a/Engine/Source/Editor/LevelEditor/Private/SLevelEditor.h b/Engine/Source/Editor/LevelEditor/Private/SLevelEditor.h
--- a/Engine/Source/Editor/LevelEditor/Private/SLevelEditor.h
+++ b/Engine/Source/Editor/LevelEditor/Private/SLevelEditor.h
@@ -146,7 +146,9 @@ private:
 	TSharedRef<SDockTab> SpawnLevelEditorTab(const FSpawnTabArgs& Args, FName TabIdentifier, FString InitializationPayload);
 	bool CanSpawnEditorModeToolbarTab(const FSpawnTabArgs& Args) const;
 	bool CanSpawnEditorModeToolboxTab(const FSpawnTabArgs& Args) const;
-	bool CanSpawnLayerBrowser(const FSpawnTabArgs& Args) const;
+	bool CanSpawnWorldBrowser(const FSpawnTabArgs& Args) const;
+	bool CanSpawnWorldPartitionBrowser(const FSpawnTabArgs& Args) const;
+	bool CanSpawnLayerBrowser(const FSpawnTabArgs& Args) const;	
 	bool CanSpawnDataLayerBrowser(const FSpawnTabArgs& Args) const;
 	bool HasAnyHostedEditorModeToolkit() const;
 
diff --git a/Engine/Source/Editor/WorldBrowser/Private/LevelCollectionModel.h b/Engine/Source/Editor/WorldBrowser/Private/LevelCollectionModel.h
--- a/Engine/Source/Editor/WorldBrowser/Private/LevelCollectionModel.h
+++ b/Engine/Source/Editor/WorldBrowser/Private/LevelCollectionModel.h
@@ -177,9 +177,6 @@ public:
 	/** @return	Whether this level collection model is a tile world */
 	virtual bool IsTileWorld() const { return false; };
 
-	/** @return	Whether this level collection model is a partition world */
-	virtual bool IsPartitionWorld() const { return false; };
-
 	/** Returns true if this collection model will support folders */
 	virtual bool HasFolderSupport() const { return false; }
 
diff --git a/Engine/Source/Editor/WorldBrowser/Private/SWorldHierarchy.cpp b/Engine/Source/Editor/WorldBrowser/Private/SWorldHierarchy.cpp
--- a/Engine/Source/Editor/WorldBrowser/Private/SWorldHierarchy.cpp
+++ b/Engine/Source/Editor/WorldBrowser/Private/SWorldHierarchy.cpp
@@ -7,6 +7,7 @@
 #include "Textures/SlateIcon.h"
 #include "Framework/Commands/UIAction.h"
 #include "Framework/MultiBox/MultiBoxBuilder.h"
+#include "WorldPartition/WorldPartitionSubsystem.h"
 #include "Widgets/Layout/SBorder.h"
 #include "Widgets/Images/SImage.h"
 #include "Widgets/Text/STextBlock.h"
@@ -59,116 +60,120 @@ void SWorldHierarchy::OnBrowseWorld(UWorld* InWorld)
 	// Bind to a new world
 	if (InWorld)
 	{
-		FWorldBrowserModule& WorldBrowserModule = FModuleManager::GetModuleChecked<FWorldBrowserModule>("WorldBrowser");
-		WorldModel = WorldBrowserModule.SharedWorldModel(InWorld);
+		if (!UWorld::HasSubsystem<UWorldPartitionSubsystem>(InWorld))
+		{
 
-		ChildSlot
-		[
-			SNew(SVerticalBox)
+			FWorldBrowserModule& WorldBrowserModule = FModuleManager::GetModuleChecked<FWorldBrowserModule>("WorldBrowser");
+			WorldModel = WorldBrowserModule.SharedWorldModel(InWorld);
 
-			// Toolbar
-			+SVerticalBox::Slot()
-			.AutoHeight()
+			ChildSlot
 			[
-				SNew(SBorder)
-				.BorderImage(FEditorStyle::GetBrush(TEXT("ToolPanel.GroupBorder")))
+				SNew(SVerticalBox)
+
+				// Toolbar
+				+SVerticalBox::Slot()
+				.AutoHeight()
 				[
-					SNew(SHorizontalBox)
-					
-					// Toolbar
-					+SHorizontalBox::Slot()
-					.AutoWidth()
-					.VAlign(VAlign_Center)
-					.HAlign(HAlign_Left)
+					SNew(SBorder)
+					.BorderImage(FEditorStyle::GetBrush(TEXT("ToolPanel.GroupBorder")))
 					[
-						// Levels menu
-						SNew( SComboButton )
-						.ComboButtonStyle(FEditorStyle::Get(), "ToolbarComboButton")
-						.ForegroundColor(FLinearColor::White)
-						.ContentPadding(0)
-						.OnGetMenuContent(this, &SWorldHierarchy::GetFileButtonContent)
-						//.ToolTipText(this, &SWorldHierarchy::GetNewAssetToolTipText)
-						//.IsEnabled(this, &SWorldHierarchy::IsAssetPathSelected )
-						.ButtonContent()
+						SNew(SHorizontalBox)
+					
+						// Toolbar
+						+SHorizontalBox::Slot()
+						.AutoWidth()
+						.VAlign(VAlign_Center)
+						.HAlign(HAlign_Left)
 						[
-							SNew(SHorizontalBox)
-
-							// Icon
-							+ SHorizontalBox::Slot()
-							.AutoWidth()
-							.VAlign(VAlign_Center)
+							// Levels menu
+							SNew( SComboButton )
+							.ComboButtonStyle(FEditorStyle::Get(), "ToolbarComboButton")
+							.ForegroundColor(FLinearColor::White)
+							.ContentPadding(0)
+							.OnGetMenuContent(this, &SWorldHierarchy::GetFileButtonContent)
+							//.ToolTipText(this, &SWorldHierarchy::GetNewAssetToolTipText)
+							//.IsEnabled(this, &SWorldHierarchy::IsAssetPathSelected )
+							.ButtonContent()
 							[
-								SNew(SImage)
-								.Image(this, &SWorldHierarchy::GetLevelsMenuBrush)
-							]
-
-							// Text
-							+ SHorizontalBox::Slot()
-							.AutoWidth()
-							.VAlign(VAlign_Center)
-							.Padding(0,0,2,0)
-							[
-								SNew(STextBlock)
-								.TextStyle(FEditorStyle::Get(), "ContentBrowser.TopBar.Font")
-								.Text(LOCTEXT("LevelsButton", "Levels"))
+								SNew(SHorizontalBox)
+
+								// Icon
+								+ SHorizontalBox::Slot()
+								.AutoWidth()
+								.VAlign(VAlign_Center)
+								[
+									SNew(SImage)
+									.Image(this, &SWorldHierarchy::GetLevelsMenuBrush)
+								]
+
+								// Text
+								+ SHorizontalBox::Slot()
+								.AutoWidth()
+								.VAlign(VAlign_Center)
+								.Padding(0,0,2,0)
+								[
+									SNew(STextBlock)
+									.TextStyle(FEditorStyle::Get(), "ContentBrowser.TopBar.Font")
+									.Text(LOCTEXT("LevelsButton", "Levels"))
+								]
 							]
 						]
-					]
 
-					// Button to summon level details tab
-					+SHorizontalBox::Slot()
-					.AutoWidth()
-					.VAlign(VAlign_Center)
-					.HAlign(HAlign_Left)
-					[
-						SNew(SButton)
-						.ButtonStyle(FEditorStyle::Get(), "ToggleButton")
-						.OnClicked(this, &SWorldHierarchy::OnSummonDetails)
-						.ToolTipText(LOCTEXT("SummonDetailsToolTipText", "Summons level details"))
-						.HAlign(HAlign_Center)
+						// Button to summon level details tab
+						+SHorizontalBox::Slot()
+						.AutoWidth()
 						.VAlign(VAlign_Center)
-						.Content()
+						.HAlign(HAlign_Left)
 						[
-							SNew(SImage)
-							.Image(this, &SWorldHierarchy::GetSummonDetailsBrush)
+							SNew(SButton)
+							.ButtonStyle(FEditorStyle::Get(), "ToggleButton")
+							.OnClicked(this, &SWorldHierarchy::OnSummonDetails)
+							.ToolTipText(LOCTEXT("SummonDetailsToolTipText", "Summons level details"))
+							.HAlign(HAlign_Center)
+							.VAlign(VAlign_Center)
+							.Content()
+							[
+								SNew(SImage)
+								.Image(this, &SWorldHierarchy::GetSummonDetailsBrush)
+							]
 						]
-					]
 
-					// Button to summon world composition tab
-					+SHorizontalBox::Slot()
-					.AutoWidth()
-					.VAlign(VAlign_Center)
-					.HAlign(HAlign_Left)
-					[
-						SNew(SButton)
-						.Visibility(this, &SWorldHierarchy::GetCompositionButtonVisibility)
-						.ButtonStyle(FEditorStyle::Get(), "ToggleButton")
-						.OnClicked(this, &SWorldHierarchy::OnSummonComposition)
-						.ToolTipText(LOCTEXT("SummonCompositionToolTipText", "Summons world composition"))
-						.HAlign(HAlign_Center)
+						// Button to summon world composition tab
+						+SHorizontalBox::Slot()
+						.AutoWidth()
 						.VAlign(VAlign_Center)
-						.Content()
+						.HAlign(HAlign_Left)
 						[
-							SNew(SImage)
-							.Image(this, &SWorldHierarchy::GetSummonCompositionBrush)
+							SNew(SButton)
+							.Visibility(this, &SWorldHierarchy::GetCompositionButtonVisibility)
+							.ButtonStyle(FEditorStyle::Get(), "ToggleButton")
+							.OnClicked(this, &SWorldHierarchy::OnSummonComposition)
+							.ToolTipText(LOCTEXT("SummonCompositionToolTipText", "Summons world composition"))
+							.HAlign(HAlign_Center)
+							.VAlign(VAlign_Center)
+							.Content()
+							[
+								SNew(SImage)
+								.Image(this, &SWorldHierarchy::GetSummonCompositionBrush)
+							]
 						]
 					]
 				]
-			]
 			
-			// Hierarchy
-			+SVerticalBox::Slot()
-			.FillHeight(1.f)
-			.Padding(0,4,0,0)
-			[
-				SNew(SBorder)
-				.BorderImage(FEditorStyle::GetBrush(TEXT("ToolPanel.GroupBorder")))
+				// Hierarchy
+				+SVerticalBox::Slot()
+				.FillHeight(1.f)
+				.Padding(0,4,0,0)
 				[
-					SNew(SWorldHierarchyImpl)
-						.InWorldModel(WorldModel)
+					SNew(SBorder)
+					.BorderImage(FEditorStyle::GetBrush(TEXT("ToolPanel.GroupBorder")))
+					[
+						SNew(SWorldHierarchyImpl)
+							.InWorldModel(WorldModel)
+					]
 				]
-			]
-		];
+			];
+		}
 	}
 }
 
diff --git a/Engine/Source/Editor/WorldBrowser/Private/StreamingLevels/StreamingLevelCollectionModel.cpp b/Engine/Source/Editor/WorldBrowser/Private/StreamingLevels/StreamingLevelCollectionModel.cpp
--- a/Engine/Source/Editor/WorldBrowser/Private/StreamingLevels/StreamingLevelCollectionModel.cpp
+++ b/Engine/Source/Editor/WorldBrowser/Private/StreamingLevels/StreamingLevelCollectionModel.cpp
@@ -418,11 +418,6 @@ void FStreamingLevelCollectionModel::UnregisterDetailsCustomization(FPropertyEdi
 	InDetailsView->UnregisterInstancedCustomPropertyLayout(ULevelStreaming::StaticClass());
 }
 
-bool FStreamingLevelCollectionModel::IsPartitionWorld() const
-{
-	return CurrentWorld.IsValid() ? !!CurrentWorld->GetWorldPartition() : false;
-}
-
 const FLevelModelList& FStreamingLevelCollectionModel::GetInvalidSelectedLevels() const 
 { 
 	return InvalidSelectedLevels;
diff --git a/Engine/Source/Editor/WorldBrowser/Private/StreamingLevels/StreamingLevelCollectionModel.h b/Engine/Source/Editor/WorldBrowser/Private/StreamingLevels/StreamingLevelCollectionModel.h
--- a/Engine/Source/Editor/WorldBrowser/Private/StreamingLevels/StreamingLevelCollectionModel.h
+++ b/Engine/Source/Editor/WorldBrowser/Private/StreamingLevels/StreamingLevelCollectionModel.h
@@ -42,7 +42,6 @@ public:
 	virtual void CustomizeFileMainMenu(FMenuBuilder& InMenuBuilder) const override;
 	virtual void RegisterDetailsCustomization(class FPropertyEditorModule& PropertyModule, TSharedPtr<class IDetailsView> InDetailsView)  override;
 	virtual void UnregisterDetailsCustomization(class FPropertyEditorModule& PropertyModule, TSharedPtr<class IDetailsView> InDetailsView) override;
-	virtual bool IsPartitionWorld() const override;
 	virtual bool HasFolderSupport() const override { return true; }
 
 private:
