Fix missing icons in levels tab

diff --git a/Engine/Source/Editor/EditorStyle/Private/StarshipStyle.cpp b/Engine/Source/Editor/EditorStyle/Private/StarshipStyle.cpp
--- a/Engine/Source/Editor/EditorStyle/Private/StarshipStyle.cpp
+++ b/Engine/Source/Editor/EditorStyle/Private/StarshipStyle.cpp
@@ -1868,7 +1868,7 @@ void FStarshipEditorStyle::FStyle::SetupLevelGeneralStyles()
 		Set("Level.NotVisibleIcon16x", new CORE_IMAGE_BRUSH_SVG("Starship/Common/hidden", Icon16x16));
 		Set("Level.NotVisibleHighlightIcon16x", new CORE_IMAGE_BRUSH_SVG("Starship/Common/hidden", Icon16x16));
 
-		Set("Level.LightingScenarioIcon16x", new IMAGE_BRUSH("Icons/icon_levels_LightingScenario_16px", Icon16x16));
+		Set("Level.LightingScenarioIcon16x", new IMAGE_BRUSH_SVG("Starship/AssetIcons/PointLight_16", Icon16x16));
 		Set("Level.LightingScenarioNotIcon16x", new IMAGE_BRUSH("Icons/icon_levels_LightingScenarioNot_16px", Icon16x16));
 		Set("Level.LockedIcon16x", new IMAGE_BRUSH("Icons/icon_locked_16px", Icon16x16));
 		Set("Level.LockedHighlightIcon16x", new IMAGE_BRUSH("Icons/icon_locked_highlight_16px", Icon16x16));
diff --git a/Engine/Source/Editor/WorldBrowser/Private/SWorldHierarchyImpl.cpp b/Engine/Source/Editor/WorldBrowser/Private/SWorldHierarchyImpl.cpp
--- a/Engine/Source/Editor/WorldBrowser/Private/SWorldHierarchyImpl.cpp
+++ b/Engine/Source/Editor/WorldBrowser/Private/SWorldHierarchyImpl.cpp
@@ -122,7 +122,7 @@ void SWorldHierarchyImpl::Construct(const FArguments& InArgs)
 		/** Lighting Scenario column */
 		+ SHeaderRow::Column(HierarchyColumns::ColumnID_LightingScenario)
 			.Visibility(bFoldersOnlyMode ? EVisibility::Collapsed : EVisibility::Visible)
-			.FixedWidth( 18.0f )
+			.FixedWidth( 24.0f )
 			.HeaderContent()
 			[
 				SNew(STextBlock)
diff --git a/Engine/Source/Editor/WorldBrowser/Private/SWorldHierarchyItem.cpp b/Engine/Source/Editor/WorldBrowser/Private/SWorldHierarchyItem.cpp
--- a/Engine/Source/Editor/WorldBrowser/Private/SWorldHierarchyItem.cpp
+++ b/Engine/Source/Editor/WorldBrowser/Private/SWorldHierarchyItem.cpp
@@ -124,13 +124,14 @@ TSharedRef< SWidget > SWorldHierarchyItem::GenerateWidgetForColumn( const FName&
 			TableRowContent =
 				SAssignNew(LightingScenarioButton, SButton)
 				.ContentPadding(0)
-				.ButtonStyle(FEditorStyle::Get(), "ToggleButton")
+				.ButtonStyle(FEditorStyle::Get(), "SimpleButton")
 				.IsEnabled(this, &SWorldHierarchyItem::IsLightingScenarioEnabled)
 				.OnClicked(this, &SWorldHierarchyItem::OnToggleLightingScenario)
 				.ToolTipText(this, &SWorldHierarchyItem::GetLightingScenarioToolTip)
 				.Visibility(this, &SWorldHierarchyItem::GetLightingScenarioVisibility)
 				.HAlign(HAlign_Center)
 				.VAlign(VAlign_Center)
+				.ForegroundColor_Lambda([this]() { return IsSelected() ? FSlateColor::UseForeground() : FSlateColor::UseStyle(); })
 				.Content()
 				[
 					SNew(SImage)
@@ -144,12 +145,13 @@ TSharedRef< SWidget > SWorldHierarchyItem::GenerateWidgetForColumn( const FName&
 			TableRowContent =
 				SAssignNew(LockButton, SButton)
 				.ContentPadding(0)
-				.ButtonStyle(FEditorStyle::Get(), "ToggleButton")
+				.ButtonStyle(FEditorStyle::Get(), "SimpleButton")
 				.IsEnabled(this, &SWorldHierarchyItem::IsLockEnabled)
 				.OnClicked(this, &SWorldHierarchyItem::OnToggleLock)
 				.ToolTipText(this, &SWorldHierarchyItem::GetLevelLockToolTip)
 				.HAlign(HAlign_Center)
 				.VAlign(VAlign_Center)
+				.ForegroundColor_Lambda([this]() { return IsSelected() ? FSlateColor::UseForeground() : FSlateColor::UseStyle(); })
 				.Content()
 				[
 					SNew(SImage)
@@ -169,6 +171,7 @@ TSharedRef< SWidget > SWorldHierarchyItem::GenerateWidgetForColumn( const FName&
 				.ToolTipText(this, &SWorldHierarchyItem::GetVisibilityToolTip)
 				.HAlign(HAlign_Center)
 				.VAlign(VAlign_Center)
+				.ForegroundColor_Lambda([this]() { return IsSelected() ? FSlateColor::UseForeground() : FSlateColor::UseStyle(); })
 				.Content()
 				[
 					SNew(SImage)
@@ -203,16 +206,18 @@ TSharedRef< SWidget > SWorldHierarchyItem::GenerateWidgetForColumn( const FName&
 			TableRowContent =
 				SAssignNew(KismetButton, SButton)
 				.ContentPadding(0)
-				.ButtonStyle(FEditorStyle::Get(), "ToggleButton")
+				.ButtonStyle(FEditorStyle::Get(), "SimpleButton")
 				.IsEnabled(this, &SWorldHierarchyItem::IsKismetEnabled)
 				.OnClicked(this, &SWorldHierarchyItem::OnOpenKismet)
 				.ToolTipText(this, &SWorldHierarchyItem::GetKismetToolTip)
 				.HAlign(HAlign_Center)
 				.VAlign(VAlign_Center)
+				.ForegroundColor_Lambda([this]() { return IsSelected() ? FSlateColor::UseForeground() : FSlateColor::UseStyle(); })
 				.Content()
 				[
 					SNew(SImage)
 					.Image(this, &SWorldHierarchyItem::GetLevelKismetBrush)
+					.ColorAndOpacity(FSlateColor::UseForeground())
 				]
 			;
 		}
@@ -242,16 +247,18 @@ TSharedRef< SWidget > SWorldHierarchyItem::GenerateWidgetForColumn( const FName&
 			TableRowContent =
 				SAssignNew(SaveButton, SButton)
 				.ContentPadding(0)
-				.ButtonStyle(FEditorStyle::Get(), "ToggleButton")
+				.ButtonStyle(FEditorStyle::Get(), "SimpleButton")
 				.IsEnabled(this, &SWorldHierarchyItem::IsSaveEnabled)
 				.OnClicked(this, &SWorldHierarchyItem::OnSave)
 				.ToolTipText(this, &SWorldHierarchyItem::GetSaveToolTip)
 				.HAlign(HAlign_Center)
 				.VAlign(VAlign_Center)
+				.ForegroundColor_Lambda([this]() { return IsSelected() ? FSlateColor::UseForeground() : FSlateColor::UseStyle(); })
 				.Content()
 				[
 					SNew(SImage)
 					.Image(this, &SWorldHierarchyItem::GetLevelSaveBrush)
+					.ColorAndOpacity(FSlateColor::UseForeground())
 				]
 			;
 		}
@@ -710,29 +717,17 @@ const FSlateBrush* SWorldHierarchyItem::GetLevelLockBrush() const
 	if (!WorldTreeItem->HasLockControls())
 	{
 		//Locking the persistent level is not allowed; stub in a different brush
-		return FEditorStyle::GetBrush( "Level.EmptyIcon16x" );
+		return FStyleDefaults::GetNoBrush();
 	}
 	else 
 	{
-		////Non-Persistent
-		//if ( GEngine && GEngine->bLockReadOnlyLevels )
-		//{
-		//	if(LevelModel->IsReadOnly())
-		//	{
-		//		return LockButton->IsHovered() ? FEditorStyle::GetBrush( "Level.ReadOnlyLockedHighlightIcon16x" ) :
-		//											FEditorStyle::GetBrush( "Level.ReadOnlyLockedIcon16x" );
-		//	}
-		//}
-			
 		if (WorldTreeItem->IsLocked())
 		{
-			return LockButton->IsHovered() ? FEditorStyle::GetBrush( "Level.LockedHighlightIcon16x" ) :
-												FEditorStyle::GetBrush( "Level.LockedIcon16x" );
+			return FEditorStyle::GetBrush("Icons.Lock");
 		}
 		else
 		{
-			return LockButton->IsHovered() ? FEditorStyle::GetBrush( "Level.UnlockedHighlightIcon16x" ) :
-												FEditorStyle::GetBrush( "Level.UnlockedIcon16x" );
+			return  FEditorStyle::GetBrush("Icons.Unlock");
 		}
 	}
 }
@@ -776,27 +771,11 @@ const FSlateBrush* SWorldHierarchyItem::GetLevelSaveBrush() const
 {
 	if (WorldTreeItem->IsLoaded())
 	{
-		if (WorldTreeItem->Flags.bLocked)
-		{
-			return FEditorStyle::GetBrush( "Level.SaveDisabledIcon16x" );
-		}
-		else
-		{
-			if (WorldTreeItem->IsDirty())
-			{
-				return SaveButton->IsHovered() ? FEditorStyle::GetBrush( "Level.SaveModifiedHighlightIcon16x" ) :
-													FEditorStyle::GetBrush( "Level.SaveModifiedIcon16x" );
-			}
-			else
-			{
-				return SaveButton->IsHovered() ? FEditorStyle::GetBrush( "Level.SaveHighlightIcon16x" ) :
-													FEditorStyle::GetBrush( "Level.SaveIcon16x" );
-			}
-		}								
+		return FEditorStyle::GetBrush("Icons.Save");							
 	}
 	else
 	{
-		return FEditorStyle::GetBrush( "Level.EmptyIcon16x" );
+		return FStyleDefaults::GetNoBrush();
 	}	
 }
 
@@ -806,17 +785,16 @@ const FSlateBrush* SWorldHierarchyItem::GetLevelKismetBrush() const
 	{
 		if (WorldTreeItem->HasKismet())
 		{
-			return KismetButton->IsHovered() ? FEditorStyle::GetBrush( "Level.ScriptHighlightIcon16x" ) :
-												FEditorStyle::GetBrush( "Level.ScriptIcon16x" );
+			return FEditorStyle::GetBrush("LevelEditor.OpenLevelBlueprint");
 		}
 		else
 		{
-			return FEditorStyle::GetBrush( "Level.EmptyIcon16x" );
+			return FStyleDefaults::GetNoBrush();
 		}
 	}
 	else
 	{
-		return FEditorStyle::GetBrush( "Level.EmptyIcon16x" );
+		return FStyleDefaults::GetNoBrush();
 	}	
 }
 
