Fix no visible dirty flag on levels in the level browser

diff --git a/Engine/Source/Editor/WorldBrowser/Private/SWorldHierarchyItem.cpp b/Engine/Source/Editor/WorldBrowser/Private/SWorldHierarchyItem.cpp
--- a/Engine/Source/Editor/WorldBrowser/Private/SWorldHierarchyItem.cpp
+++ b/Engine/Source/Editor/WorldBrowser/Private/SWorldHierarchyItem.cpp
@@ -267,10 +267,11 @@ FText SWorldHierarchyItem::GetDisplayNameText() const
 {
 	FFormatNamedArguments Args;
 	Args.Add(TEXT("DisplayText"), FText::FromString(WorldTreeItem->GetDisplayString()));
+	Args.Add(TEXT("DirtyFlag"), WorldTreeItem->IsDirty() ? FText::FromString("*") : FText::GetEmpty());
 	Args.Add(TEXT("Transient"), (WorldTreeItem->IsTransient() ? LOCTEXT("WorldItem_Transient", " (Transient)") : FText::GetEmpty()));
 	Args.Add(TEXT("ReadOnly"), (WorldTreeItem->IsReadOnly() ? LOCTEXT("WorldItem_ReadOnly", " (Read-Only)") : FText::GetEmpty()));
 
-	return FText::Format(LOCTEXT("WorldItem", "{DisplayText}{Transient}{ReadOnly}"), Args);
+	return FText::Format(LOCTEXT("WorldItem", "{DisplayText}{DirtyFlag} {Transient}{ReadOnly}"), Args);
 }
 
 void SWorldHierarchyItem::OnLabelCommitted(const FText& InLabel, ETextCommit::Type InCommitInfo)
