Adds "mark for add" when saving new levels in the "Levels" Windows

#changelist validated

diff --git a/Engine/Source/Editor/WorldBrowser/Private/LevelCollectionModel.cpp b/Engine/Source/Editor/WorldBrowser/Private/LevelCollectionModel.cpp
--- a/Engine/Source/Editor/WorldBrowser/Private/LevelCollectionModel.cpp
+++ b/Engine/Source/Editor/WorldBrowser/Private/LevelCollectionModel.cpp
@@ -627,6 +627,12 @@ void FLevelCollectionModel::SaveLevels(const FLevelModelList& InLevelList)
 		{
 			FEditorFileUtils::SaveLevel(*It);
 		}
+
+		// Add all files that needs to be marked for add in one command, if any
+		if (GEditor)
+		{
+			GEditor->RunDeferredMarkForAddFiles();
+		}
 	}
 	else if (PackagesNotNeedingCheckout.Num() > 0)
 	{
