Fixed "Levels" UI empty on New Projects /w World Partition - It mentions it explicitely now

diff --git a/Engine/Source/Editor/WorldBrowser/Private/SWorldHierarchy.cpp b/Engine/Source/Editor/WorldBrowser/Private/SWorldHierarchy.cpp
--- a/Engine/Source/Editor/WorldBrowser/Private/SWorldHierarchy.cpp
+++ b/Engine/Source/Editor/WorldBrowser/Private/SWorldHierarchy.cpp
@@ -60,7 +60,21 @@ void SWorldHierarchy::OnBrowseWorld(UWorld* InWorld)
 	// Bind to a new world
 	if (InWorld)
 	{
-		if (!UWorld::HasSubsystem<UWorldPartitionSubsystem>(InWorld))
+		if (UWorld::HasSubsystem<UWorldPartitionSubsystem>(InWorld))
+		{
+			ChildSlot
+			[
+				SNew(SHorizontalBox)
+				+ SHorizontalBox::Slot()
+				.VAlign(VAlign_Center)
+				.HAlign(HAlign_Center)
+				[
+					SNew(STextBlock)
+					.Text(LOCTEXT("FeatureDisabledWithWorldPartition", "This feature is disabled when World Partition is enabled."))
+				]
+			];
+		}
+		else
 		{
 
 			FWorldBrowserModule& WorldBrowserModule = FModuleManager::GetModuleChecked<FWorldBrowserModule>("WorldBrowser");
