Levels window now only shows a save modified icon if a level is dirty,
Fixed an issue where new folder default names would not increment correctly,
Updated lock icons and added save modified and solid light bulb icons.

diff --git a/Engine/Source/Editor/EditorStyle/Private/StarshipStyle.cpp b/Engine/Source/Editor/EditorStyle/Private/StarshipStyle.cpp
--- a/Engine/Source/Editor/EditorStyle/Private/StarshipStyle.cpp
+++ b/Engine/Source/Editor/EditorStyle/Private/StarshipStyle.cpp
@@ -1840,6 +1840,8 @@ void FStarshipEditorStyle::FStyle::SetupLevelGeneralStyles()
 
 		Set("Level.LightingScenarioIcon16x", new IMAGE_BRUSH_SVG("Starship/AssetIcons/PointLight_16", Icon16x16));
 		Set("Level.LightingScenarioNotIcon16x", new IMAGE_BRUSH_SVG("Starship/Common/LightBulbOff", Icon16x16));
+		Set("Level.LightingScenarioIconSolid16x", new IMAGE_BRUSH_SVG("Starship/Common/LightBulbSolid", Icon16x16));
+		Set("Level.LightingScenarioNotIconSolid16x", new IMAGE_BRUSH_SVG("Starship/Common/LightBulbOffSolid", Icon16x16));
 		Set("Level.LockedIcon16x", new IMAGE_BRUSH("Icons/icon_locked_16px", Icon16x16));
 		Set("Level.LockedHighlightIcon16x", new IMAGE_BRUSH("Icons/icon_locked_highlight_16px", Icon16x16));
 		Set("Level.UnlockedIcon16x", new IMAGE_BRUSH("Icons/icon_levels_unlocked_16px", Icon16x16));
diff --git a/Engine/Source/Editor/WorldBrowser/Private/SWorldHierarchyItem.cpp b/Engine/Source/Editor/WorldBrowser/Private/SWorldHierarchyItem.cpp
--- a/Engine/Source/Editor/WorldBrowser/Private/SWorldHierarchyItem.cpp
+++ b/Engine/Source/Editor/WorldBrowser/Private/SWorldHierarchyItem.cpp
@@ -112,6 +112,7 @@ TSharedRef< SWidget > SWorldHierarchyItem::GenerateWidgetForColumn( const FName&
 			+SHorizontalBox::Slot()
 			.VAlign(VAlign_Center)
 			.AutoWidth()
+			.Padding(FMargin(5.f, 0.f, 0.f, 0.f))
 			[
 				TextWidget.ToSharedRef()
 			]
@@ -126,7 +127,7 @@ TSharedRef< SWidget > SWorldHierarchyItem::GenerateWidgetForColumn( const FName&
 			TableRowContent =
 				SAssignNew(LightingScenarioButton, SButton)
 				.ContentPadding(0)
-				.ButtonStyle(FEditorStyle::Get(), "SimpleButton")
+				.ButtonStyle(FAppStyle::Get(), "SimpleButton")
 				.IsEnabled(this, &SWorldHierarchyItem::IsLightingScenarioEnabled)
 				.OnClicked(this, &SWorldHierarchyItem::OnToggleLightingScenario)
 				.ToolTipText(this, &SWorldHierarchyItem::GetLightingScenarioToolTip)
@@ -147,7 +148,7 @@ TSharedRef< SWidget > SWorldHierarchyItem::GenerateWidgetForColumn( const FName&
 			TableRowContent =
 				SAssignNew(LockButton, SButton)
 				.ContentPadding(0)
-				.ButtonStyle(FEditorStyle::Get(), "SimpleButton")
+				.ButtonStyle(FAppStyle::Get(), "SimpleButton")
 				.IsEnabled(this, &SWorldHierarchyItem::IsLockEnabled)
 				.OnClicked(this, &SWorldHierarchyItem::OnToggleLock)
 				.ToolTipText(this, &SWorldHierarchyItem::GetLevelLockToolTip)
@@ -167,7 +168,7 @@ TSharedRef< SWidget > SWorldHierarchyItem::GenerateWidgetForColumn( const FName&
 			TableRowContent =
 				SAssignNew(VisibilityButton, SButton)
 				.ContentPadding(0)
-				.ButtonStyle(FEditorStyle::Get(), "SimpleButton")
+				.ButtonStyle(FAppStyle::Get(), "SimpleButton")
 				.IsEnabled(this, &SWorldHierarchyItem::IsVisibilityEnabled)
 				.OnClicked(this, &SWorldHierarchyItem::OnToggleVisibility)
 				.ToolTipText(this, &SWorldHierarchyItem::GetVisibilityToolTip)
@@ -187,7 +188,7 @@ TSharedRef< SWidget > SWorldHierarchyItem::GenerateWidgetForColumn( const FName&
 			TableRowContent =
 				SAssignNew(ColorButton, SButton)
 				.ContentPadding(0)
-				.ButtonStyle(FEditorStyle::Get(), "SimpleButton")
+				.ButtonStyle(FAppStyle::Get(), "SimpleButton")
 				.IsEnabled(true)
 				.OnClicked(this, &SWorldHierarchyItem::OnChangeColor)
 				.ToolTipText(LOCTEXT("LevelColorButtonToolTip", "Change Level Color"))
@@ -198,7 +199,7 @@ TSharedRef< SWidget > SWorldHierarchyItem::GenerateWidgetForColumn( const FName&
 				[
 					SNew(SImage)
 					.ColorAndOpacity(this, &SWorldHierarchyItem::GetDrawColor)
-					.Image(FEditorStyle::GetBrush("Level.ColorIcon"))
+					.Image(FAppStyle::Get().GetBrush("Level.ColorIcon"))
 				]
 			;
 		}
@@ -207,7 +208,7 @@ TSharedRef< SWidget > SWorldHierarchyItem::GenerateWidgetForColumn( const FName&
 			TableRowContent =
 				SAssignNew(KismetButton, SButton)
 				.ContentPadding(0)
-				.ButtonStyle(FEditorStyle::Get(), "SimpleButton")
+				.ButtonStyle(FAppStyle::Get(), "SimpleButton")
 				.IsEnabled(this, &SWorldHierarchyItem::IsKismetEnabled)
 				.OnClicked(this, &SWorldHierarchyItem::OnOpenKismet)
 				.ToolTipText(this, &SWorldHierarchyItem::GetKismetToolTip)
@@ -242,7 +243,7 @@ TSharedRef< SWidget > SWorldHierarchyItem::GenerateWidgetForColumn( const FName&
 			TableRowContent =
 				SAssignNew(SaveButton, SButton)
 				.ContentPadding(0)
-				.ButtonStyle(FEditorStyle::Get(), "SimpleButton")
+				.ButtonStyle(FAppStyle::Get(), "SimpleButton")
 				.IsEnabled(this, &SWorldHierarchyItem::IsSaveEnabled)
 				.OnClicked(this, &SWorldHierarchyItem::OnSave)
 				.ToolTipText(this, &SWorldHierarchyItem::GetSaveToolTip)
@@ -267,11 +268,11 @@ FText SWorldHierarchyItem::GetDisplayNameText() const
 {
 	FFormatNamedArguments Args;
 	Args.Add(TEXT("DisplayText"), FText::FromString(WorldTreeItem->GetDisplayString()));
-	Args.Add(TEXT("DirtyFlag"), WorldTreeItem->IsDirty() ? FText::FromString("*") : FText::GetEmpty());
+	Args.Add(TEXT("DirtyFlag"), WorldTreeItem->IsDirty() ? FText::FromString(" *") : FText::GetEmpty());
 	Args.Add(TEXT("Transient"), (WorldTreeItem->IsTransient() ? LOCTEXT("WorldItem_Transient", " (Transient)") : FText::GetEmpty()));
 	Args.Add(TEXT("ReadOnly"), (WorldTreeItem->IsReadOnly() ? LOCTEXT("WorldItem_ReadOnly", " (Read-Only)") : FText::GetEmpty()));
 
-	return FText::Format(LOCTEXT("WorldItem", "{DisplayText}{DirtyFlag} {Transient}{ReadOnly}"), Args);
+	return FText::Format(LOCTEXT("WorldItem", "{DisplayText}{DirtyFlag}{Transient}{ReadOnly}"), Args);
 }
 
 void SWorldHierarchyItem::OnLabelCommitted(const FText& InLabel, ETextCommit::Type InCommitInfo)
@@ -589,7 +590,7 @@ void SWorldHierarchyItem::OnDragEnter( const FGeometry& MyGeometry, const FDragD
 	WorldHierarchy::FValidationInfo ValidationInfo = WorldTreeItem->ValidateDrop(DragDropEvent);
 
 	FName IconName = ValidationInfo.bValid ? TEXT("Graph.ConnectorFeedback.OK") : TEXT("Graph.ConnectorFeedback.Error");
-	const FSlateBrush* Icon = FEditorStyle::GetBrush(IconName);
+	const FSlateBrush* Icon = FAppStyle::Get().GetBrush(IconName);
 
 	TSharedPtr<FDragDropOperation> Operation = DragDropEvent.GetOperation();
 	if (Operation.IsValid())
@@ -629,11 +630,11 @@ FSlateFontInfo SWorldHierarchyItem::GetDisplayNameFont() const
 {
 	if (WorldTreeItem->IsCurrent())
 	{
-		return FEditorStyle::GetFontStyle("WorldBrowser.LabelFontBold");
+		return FAppStyle::Get().GetFontStyle("WorldBrowser.LabelFontBold");
 	}
 	else
 	{
-		return FEditorStyle::GetFontStyle("WorldBrowser.LabelFont");
+		return FAppStyle::Get().GetFontStyle("WorldBrowser.LabelFont");
 	}
 }
 
@@ -676,18 +677,18 @@ const FSlateBrush* SWorldHierarchyItem::GetLevelVisibilityBrush() const
 	{
 		if (WorldTreeItem->IsVisible())
 		{
-			return VisibilityButton->IsHovered() ? FEditorStyle::GetBrush( "Level.VisibleHighlightIcon16x" ) :
-													FEditorStyle::GetBrush( "Level.VisibleIcon16x" );
+			return VisibilityButton->IsHovered() ? FAppStyle::Get().GetBrush( "Level.VisibleHighlightIcon16x" ) :
+													FAppStyle::Get().GetBrush( "Level.VisibleIcon16x" );
 		}
 		else
 		{
-			return VisibilityButton->IsHovered() ? FEditorStyle::GetBrush( "Level.NotVisibleHighlightIcon16x" ) :
-													FEditorStyle::GetBrush( "Level.NotVisibleIcon16x" );
+			return VisibilityButton->IsHovered() ? FAppStyle::Get().GetBrush( "Level.NotVisibleHighlightIcon16x" ) :
+													FAppStyle::Get().GetBrush( "Level.NotVisibleIcon16x" );
 		}
 	}
 	else
 	{
-		return FEditorStyle::GetBrush( "Level.EmptyIcon16x" );
+		return FAppStyle::Get().GetBrush( "Level.EmptyIcon16x" );
 	}
 }
 
@@ -695,11 +696,11 @@ const FSlateBrush* SWorldHierarchyItem::GetLightingScenarioBrush() const
 {
 	if (WorldTreeItem->IsLightingScenario())
 	{
-		return FEditorStyle::GetBrush( "Level.LightingScenarioIcon16x" );
+		return FAppStyle::Get().GetBrush( "Level.LightingScenarioIconSolid16x" );
 	}
 	else
 	{
-		return FEditorStyle::GetBrush( "Level.LightingScenarioNotIcon16x" );
+		return FAppStyle::Get().GetBrush( "Level.LightingScenarioNotIconSolid16x" );
 	}
 }
 
@@ -719,11 +720,11 @@ const FSlateBrush* SWorldHierarchyItem::GetLevelLockBrush() const
 	{
 		if (WorldTreeItem->IsLocked())
 		{
-			return FEditorStyle::GetBrush("Icons.Lock");
+			return FAppStyle::Get().GetBrush("Icons.Lock");
 		}
 		else
 		{
-			return  FEditorStyle::GetBrush("Icons.Unlock");
+			return  FAppStyle::Get().GetBrush("Icons.Unlock");
 		}
 	}
 }
@@ -765,14 +766,14 @@ const FSlateBrush* SWorldHierarchyItem::GetSCCStateImage() const
 
 const FSlateBrush* SWorldHierarchyItem::GetLevelSaveBrush() const
 {
-	if (WorldTreeItem->IsLoaded())
+	if (WorldTreeItem->IsDirty())
 	{
-		return FEditorStyle::GetBrush("Icons.Save");							
+		return FAppStyle::Get().GetBrush("Icons.SaveModified");
 	}
 	else
 	{
 		return FStyleDefaults::GetNoBrush();
-	}	
+	}
 }
 
 const FSlateBrush* SWorldHierarchyItem::GetLevelKismetBrush() const
@@ -781,7 +782,7 @@ const FSlateBrush* SWorldHierarchyItem::GetLevelKismetBrush() const
 	{
 		if (WorldTreeItem->HasKismet())
 		{
-			return FEditorStyle::GetBrush("LevelEditor.OpenLevelBlueprint");
+			return FAppStyle::Get().GetBrush("Icons.Blueprints");
 		}
 		else
 		{
diff --git a/Engine/Source/Runtime/SlateCore/Private/Styling/StarshipCoreStyle.cpp b/Engine/Source/Runtime/SlateCore/Private/Styling/StarshipCoreStyle.cpp
--- a/Engine/Source/Runtime/SlateCore/Private/Styling/StarshipCoreStyle.cpp
+++ b/Engine/Source/Runtime/SlateCore/Private/Styling/StarshipCoreStyle.cpp
@@ -333,6 +333,7 @@ TSharedRef<ISlateStyle> FStarshipCoreStyle::Create()
 		Style->Set("Icons.X", new IMAGE_BRUSH_SVG("Starship/Common/close", Icon16x16));
 		Style->Set("Icons.Delete", new IMAGE_BRUSH_SVG("Starship/Common/delete-outline", Icon16x16));
 		Style->Set("Icons.Save", new IMAGE_BRUSH_SVG("Starship/Common/save", Icon16x16));
+		Style->Set("Icons.SaveModified", new IMAGE_BRUSH_SVG("Starship/Common/save-modified", Icon16x16));
 
 		Style->Set("Icons.Import", new IMAGE_BRUSH_SVG("Starship/Common/import", Icon16x16));
 		Style->Set("Icons.Filter", new IMAGE_BRUSH_SVG("Starship/Common/filter", Icon16x16));
