Checks if levels/actors are dirty before saving levels. This avoids checkouting all external actor when OFPA is enabled.

#changelist validated

diff --git a/Engine/Source/Editor/WorldBrowser/Private/LevelCollectionModel.cpp b/Engine/Source/Editor/WorldBrowser/Private/LevelCollectionModel.cpp
--- a/Engine/Source/Editor/WorldBrowser/Private/LevelCollectionModel.cpp
+++ b/Engine/Source/Editor/WorldBrowser/Private/LevelCollectionModel.cpp
@@ -1,6 +1,7 @@
 // Copyright Epic Games, Inc. All Rights Reserved.
 
 #include "LevelCollectionModel.h"
+#include "Algo/AnyOf.h"
 #include "Misc/PackageName.h"
 #include "AssetData.h"
 #include "Misc/MessageDialog.h"
@@ -618,8 +619,12 @@ void FLevelCollectionModel::SaveLevels(const FLevelModelList& InLevelList)
 	}
 
 	TArray< UPackage* > PackagesNotNeedingCheckout;
+
+	// Check dirtiness in case of level using external actors to avoid taking in checkout all actors
+	bool bCheckDirty = Algo::AnyOf(LevelsToSave, [](const ULevel* InLevel) -> bool { return (InLevel != nullptr) && InLevel->IsUsingExternalActors(); });
+
 	// Prompt the user to check out the levels from source control before saving
-	if (FEditorFileUtils::PromptToCheckoutLevels(false, LevelsToSave, &PackagesNotNeedingCheckout))
+	if (FEditorFileUtils::PromptToCheckoutLevels(bCheckDirty, LevelsToSave, &PackagesNotNeedingCheckout))
 	{
 		for (auto It = LevelsToSave.CreateIterator(); It; ++It)
 		{
