Expand any parent items when selecting a level.

diff --git a/Engine/Source/Editor/WorldBrowser/Private/SWorldHierarchyImpl.cpp b/Engine/Source/Editor/WorldBrowser/Private/SWorldHierarchyImpl.cpp
--- a/Engine/Source/Editor/WorldBrowser/Private/SWorldHierarchyImpl.cpp
+++ b/Engine/Source/Editor/WorldBrowser/Private/SWorldHierarchyImpl.cpp
@@ -1075,6 +1075,12 @@ void SWorldHierarchyImpl::Populate()
 			if (TreeItemMap.Contains(ID))
 			{
 				WorldHierarchy::FWorldTreeItemPtr Item = TreeItemMap[ID];
+
+				for (WorldHierarchy::FWorldTreeItemPtr ItemParent = Item->GetParent(); ItemParent.IsValid(); ItemParent = ItemParent->GetParent())
+				{
+					TreeWidget->SetItemExpansion(ItemParent, true);
+				}
+
 				TreeWidget->SetItemSelection(Item, true);
 
 				if (!bScrolledIntoView)
