Ensure that the world is marked dirty when removing a sub-level.

diff --git a/Engine/Source/Editor/WorldBrowser/Private/LevelCollectionModel.cpp b/Engine/Source/Editor/WorldBrowser/Private/LevelCollectionModel.cpp
--- a/Engine/Source/Editor/WorldBrowser/Private/LevelCollectionModel.cpp
+++ b/Engine/Source/Editor/WorldBrowser/Private/LevelCollectionModel.cpp
@@ -737,12 +737,8 @@ void FLevelCollectionModel::UnloadLevels(const FLevelModelList& InLevelList)
 					ThisWorld->RemoveStreamingLevel(*StreamingLevel);
 				}
 			}
-			
-			// Unload sub-level
-			{
-				FUnmodifiableObject ImmuneWorld(CurrentWorld.Get());
-				EditorLevelUtils::RemoveLevelFromWorld(Level);
-			}
+
+			EditorLevelUtils::RemoveLevelFromWorld(Level);
 		}
 		else if (ULevelStreaming* StreamingLevel = Cast<ULevelStreaming>(LevelModel->GetNodeObject()))
 		{
