Prevent crash when unloading a World Composition level whose package cannot be unloaded because there's at least one reference to it preventing the world to be garbage-collected. This test allows the WorldBrowser to survive until the GCObjectReferencer can run and display the actual reason preventing garbage collection (via a fatal error)

diff --git a/Engine/Source/Editor/WorldBrowser/Private/Tiles/WorldTileModel.cpp b/Engine/Source/Editor/WorldBrowser/Private/Tiles/WorldTileModel.cpp
--- a/Engine/Source/Editor/WorldBrowser/Private/Tiles/WorldTileModel.cpp
+++ b/Engine/Source/Editor/WorldBrowser/Private/Tiles/WorldTileModel.cpp
@@ -70,7 +70,7 @@ FWorldTileModel::FWorldTileModel(FWorldTileCollectionModel& InWorldModel, int32
 		{
 			// Find the world object
 			UWorld* World = UWorld::FindWorldInPackage(LevelPackage);
-			if (World)
+			if (IsValid(World))
 			{
 				LoadedLevel = World->PersistentLevel;
 				// Enable tile properties
