Fix Level Instance Editing Level not showing up in Current Level drop down

- Edited Level Instance LevelStreaming now part of the LevelCollectionModel
- Add a IsUserManaged() virtual to LevelStreaming to limit the actions that can be done on some LevelStreaming classes
- LevelInstanceEditorLevelStreaming overrides IsUserManaged() to return false

#rnx

diff --git a/Engine/Source/Editor/WorldBrowser/Private/LevelCollectionModel.cpp b/Engine/Source/Editor/WorldBrowser/Private/LevelCollectionModel.cpp
--- a/Engine/Source/Editor/WorldBrowser/Private/LevelCollectionModel.cpp
+++ b/Engine/Source/Editor/WorldBrowser/Private/LevelCollectionModel.cpp
@@ -136,7 +136,7 @@ void FLevelCollectionModel::BindCommands()
 
 	ActionList.MapAction(Commands.World_UnloadLevel,
 		FExecuteAction::CreateSP(this, &FLevelCollectionModel::UnloadSelectedLevels_Executed),
-		FCanExecuteAction::CreateSP(this, &FLevelCollectionModel::AreAnySelectedLevelsLoaded));
+		FCanExecuteAction::CreateSP(this, &FLevelCollectionModel::AreAllSelectedLevelsUserManaged));
 
 	ActionList.MapAction( Commands.World_MigrateSelectedLevels,
 		FExecuteAction::CreateSP( this, &FLevelCollectionModel::MigrateSelectedLevels_Executed),
@@ -715,7 +715,7 @@ void FLevelCollectionModel::UnloadLevels(const FLevelModelList& InLevelList)
 		TSharedPtr<FLevelModel> LevelModel = (*It);
 		ULevel* Level = LevelModel->GetLevelObject();
 
-		if (Level != nullptr && !LevelModel->IsPersistent())
+		if (Level != nullptr && !LevelModel->IsPersistent() && LevelModel->IsUserManaged())
 		{
 			// Unselect all actors before removing the level
 			// This avoids crashing in areas that rely on getting a selected actors level. The level will be invalid after its removed.
@@ -931,6 +931,19 @@ bool FLevelCollectionModel::AreAnyLevelsSelected() const
 	return SelectedLevelsList.Num() > 0;
 }
 
+bool FLevelCollectionModel::AreAllSelectedLevelsUserManaged() const
+{
+	for (int32 LevelIdx = 0; LevelIdx < SelectedLevelsList.Num(); LevelIdx++)
+	{
+		if (!SelectedLevelsList[LevelIdx]->IsUserManaged())
+		{
+			return false;
+		}
+	}
+
+	return AreAnyLevelsSelected();
+}
+
 bool FLevelCollectionModel::AreAllSelectedLevelsLoaded() const
 {
 	for (int32 LevelIdx = 0; LevelIdx < SelectedLevelsList.Num(); LevelIdx++)
@@ -984,7 +997,7 @@ bool FLevelCollectionModel::AreAllSelectedLevelsEditableAndNotPersistent() const
 {
 	for (auto It = SelectedLevelsList.CreateConstIterator(); It; ++It)
 	{
-		if ((*It)->IsEditable() == false || (*It)->IsPersistent())
+		if ((*It)->IsEditable() == false || (*It)->IsPersistent() || !(*It)->IsUserManaged())
 		{
 			return false;
 		}
diff --git a/Engine/Source/Editor/WorldBrowser/Private/LevelCollectionModel.h b/Engine/Source/Editor/WorldBrowser/Private/LevelCollectionModel.h
--- a/Engine/Source/Editor/WorldBrowser/Private/LevelCollectionModel.h
+++ b/Engine/Source/Editor/WorldBrowser/Private/LevelCollectionModel.h
@@ -207,6 +207,9 @@ public:
 	/** @return	whether at least one level is selected */
 	bool AreAnyLevelsSelected() const;
 
+	/** @return wether all selected levels are user managed */
+	bool AreAllSelectedLevelsUserManaged() const;
+
 	/** @return whether all the currently selected levels are loaded */
 	bool AreAllSelectedLevelsLoaded() const;
 
diff --git a/Engine/Source/Editor/WorldBrowser/Private/LevelModel.cpp b/Engine/Source/Editor/WorldBrowser/Private/LevelModel.cpp
--- a/Engine/Source/Editor/WorldBrowser/Private/LevelModel.cpp
+++ b/Engine/Source/Editor/WorldBrowser/Private/LevelModel.cpp
@@ -238,6 +238,20 @@ bool FLevelModel::IsFileReadOnly() const
 	return false;
 }
 
+bool FLevelModel::IsUserManaged() const
+{
+	ULevel* Level = GetLevelObject();
+	if (Level)
+	{
+		if (ULevelStreaming* StreamingLevel = FLevelUtils::FindStreamingLevel(Level))
+		{
+			return StreamingLevel->IsUserManaged();
+		}
+	}
+
+	return true;
+}
+
 void FLevelModel::LoadLevel()
 {
 
@@ -753,7 +767,7 @@ void FLevelModel::SelectActors(bool bSelect, bool bNotify, bool bSelectEvenIfHid
 void FLevelModel::ConvertLevelToExternalActors(bool bUseExternal)
 {
 	ULevel* Level = GetLevelObject();
-	if (Level == nullptr || IsLocked())
+	if (Level == nullptr || IsLocked() || !IsUserManaged())
 	{
 		return;
 	}
@@ -764,7 +778,7 @@ void FLevelModel::ConvertLevelToExternalActors(bool bUseExternal)
 bool FLevelModel::CanConvertLevelToExternalActors(bool bToExternal)
 {
 	ULevel* Level = GetLevelObject();
-	if (Level == nullptr || IsLocked())
+	if (Level == nullptr || IsLocked() || !IsUserManaged())
 	{
 		return false;
 	}
diff --git a/Engine/Source/Editor/WorldBrowser/Private/LevelModel.h b/Engine/Source/Editor/WorldBrowser/Private/LevelModel.h
--- a/Engine/Source/Editor/WorldBrowser/Private/LevelModel.h
+++ b/Engine/Source/Editor/WorldBrowser/Private/LevelModel.h
@@ -126,6 +126,9 @@ public:
 	/** @return Whether level package file is read only */
 	bool IsFileReadOnly() const;
 
+	/** @return Whether this level can be add/removed/unloaded by user */
+	bool IsUserManaged() const;
+
 	/** @return Whether the streaming level object is transient */
 	virtual bool IsTransient() const { return false; }
 
diff --git a/Engine/Source/Editor/WorldBrowser/Private/StreamingLevels/StreamingLevelCollectionModel.cpp b/Engine/Source/Editor/WorldBrowser/Private/StreamingLevels/StreamingLevelCollectionModel.cpp
--- a/Engine/Source/Editor/WorldBrowser/Private/StreamingLevels/StreamingLevelCollectionModel.cpp
+++ b/Engine/Source/Editor/WorldBrowser/Private/StreamingLevels/StreamingLevelCollectionModel.cpp
@@ -294,7 +294,7 @@ void FStreamingLevelCollectionModel::BuildHierarchyMenu(FMenuBuilder& InMenuBuil
 			FNewMenuDelegate::CreateSP(const_cast<FStreamingLevelCollectionModel*>(this), &FStreamingLevelCollectionModel::FillLockSubMenu ) );
 		
 		// Level streaming specific commands
-		if (AreAnyLevelsSelected() && !(IsOneLevelSelected() && GetSelectedLevels()[0]->IsPersistent()))
+		if (AreAllSelectedLevelsUserManaged() && !(IsOneLevelSelected() && GetSelectedLevels()[0]->IsPersistent()))
 		{
 			InMenuBuilder.AddMenuEntry(Commands.World_RemoveSelectedLevels);
 			//
@@ -339,9 +339,12 @@ void FStreamingLevelCollectionModel::BuildHierarchyMenu(FMenuBuilder& InMenuBuil
 			InMenuBuilder.AddMenuEntry( Commands.MoveFoliageToSelected );
 		}
 
-		InMenuBuilder.AddMenuEntry(Commands.ConvertLevelToExternalActors);
-		InMenuBuilder.AddMenuEntry(Commands.ConvertLevelToInternalActors);
-		
+		if (AreAllSelectedLevelsUserManaged())
+		{
+			InMenuBuilder.AddMenuEntry(Commands.ConvertLevelToExternalActors);
+			InMenuBuilder.AddMenuEntry(Commands.ConvertLevelToInternalActors);
+		}
+
 		if (AreAnyLevelsSelected() && !(IsOneLevelSelected() && SelectedLevelsList[0]->IsPersistent()))
 		{
 			InMenuBuilder.AddMenuEntry( Commands.SelectStreamingVolumes );
@@ -584,7 +587,7 @@ bool FStreamingLevelCollectionModel::AreAllSelectedLevelsRemovable() const
 {
 	for (const TSharedPtr<FLevelModel>& LevelModel : SelectedLevelsList)
 	{
-		if (LevelModel->IsLocked() || LevelModel->IsPersistent())
+		if (LevelModel->IsLocked() || LevelModel->IsPersistent() || !LevelModel->IsUserManaged())
 		{
 			return false;
 		}
diff --git a/Engine/Source/Editor/WorldBrowser/Private/WorldTreeItemTypes.cpp b/Engine/Source/Editor/WorldBrowser/Private/WorldTreeItemTypes.cpp
--- a/Engine/Source/Editor/WorldBrowser/Private/WorldTreeItemTypes.cpp
+++ b/Engine/Source/Editor/WorldBrowser/Private/WorldTreeItemTypes.cpp
@@ -473,7 +473,7 @@ namespace WorldHierarchy
 
 	bool FLevelModelTreeItem::CanChangeParents() const
 	{
-		return Parent.IsValid();
+		return Parent.IsValid() && LevelModel.IsValid() && LevelModel.Pin()->IsUserManaged();
 	}
 
 	void FLevelModelTreeItem::GenerateContextMenu(UToolMenu* Menu, const SWorldHierarchyImpl& Hierarchy)
diff --git a/Engine/Source/Runtime/Engine/Classes/Engine/LevelStreaming.h b/Engine/Source/Runtime/Engine/Classes/Engine/LevelStreaming.h
--- a/Engine/Source/Runtime/Engine/Classes/Engine/LevelStreaming.h
+++ b/Engine/Source/Runtime/Engine/Classes/Engine/LevelStreaming.h
@@ -465,6 +465,9 @@ public:
 
 	/** Returns if the streaming level is visible in LevelCollectionModel */
 	virtual bool ShowInLevelCollection() const { return true; }
+
+	/** Returns if the streaming level can be modified by user */
+	virtual bool IsUserManaged() const { return true; }
 #endif
 
 	/** Returns a constant reference to the world asset this streaming level object references  */
diff --git a/Engine/Source/Runtime/Engine/Public/LevelInstance/LevelInstanceEditorLevelStreaming.h b/Engine/Source/Runtime/Engine/Public/LevelInstance/LevelInstanceEditorLevelStreaming.h
--- a/Engine/Source/Runtime/Engine/Public/LevelInstance/LevelInstanceEditorLevelStreaming.h
+++ b/Engine/Source/Runtime/Engine/Public/LevelInstance/LevelInstanceEditorLevelStreaming.h
@@ -17,7 +17,8 @@ class ULevelStreamingLevelInstanceEditor : public ULevelStreamingAlwaysLoaded
 
 public:
 #if WITH_EDITOR
-	virtual bool ShowInLevelCollection() const override { return false; }
+	virtual bool ShowInLevelCollection() const override { return true; }
+	virtual bool IsUserManaged() const override { return false; }
 	ILevelInstanceInterface* GetLevelInstance() const;
 	const FLevelInstanceID& GetLevelInstanceID() const { return LevelInstanceID; }
 	FBox GetBounds() const;
diff --git a/Engine/Source/Runtime/Engine/Public/LevelInstance/LevelInstanceLevelStreaming.h b/Engine/Source/Runtime/Engine/Public/LevelInstance/LevelInstanceLevelStreaming.h
--- a/Engine/Source/Runtime/Engine/Public/LevelInstance/LevelInstanceLevelStreaming.h
+++ b/Engine/Source/Runtime/Engine/Public/LevelInstance/LevelInstanceLevelStreaming.h
@@ -21,6 +21,8 @@ public:
 
 #if WITH_EDITOR
 	virtual bool ShowInLevelCollection() const override { return false; }
+	virtual bool IsUserManaged() const override { return false; }
+
 	FBox GetBounds() const;
 
 	virtual TOptional<FFolder::FRootObject> GetFolderRootObject() const override;
