The problem:
When the level is locked, CanSelectElement evaluates to true, causing Unreal Engine to crash in the following two scenarios:
1. After locking the level, the user clicks on an ISM
2. After locking the level, the user drags the remaining gizmo on an ISM

The cause:
1. FSMInstanceElementLevelEditorSelectionCustomization::CanSelectElement evaluates to true because GEdSelectionLock seems always be false
2. DeselectAllActors in SetLocked does not recursively deselect the ISM components

The fix:
1. Determine whether the owner level is lockedin FSMInstanceElementLevelEditorSelectionCustomization::CanSelectElement
2. Use the typed element selection APIs to clear the current level's selection

diff --git a/Engine/Source/Editor/LevelEditor/Private/Elements/SMInstance/SMInstanceElementLevelEditorSelectionCustomization.cpp b/Engine/Source/Editor/LevelEditor/Private/Elements/SMInstance/SMInstanceElementLevelEditorSelectionCustomization.cpp
--- a/Engine/Source/Editor/LevelEditor/Private/Elements/SMInstance/SMInstanceElementLevelEditorSelectionCustomization.cpp
+++ b/Engine/Source/Editor/LevelEditor/Private/Elements/SMInstance/SMInstanceElementLevelEditorSelectionCustomization.cpp
@@ -7,6 +7,8 @@
 #include "Elements/Framework/TypedElementList.h"
 #include "Elements/Framework/EngineElementsLibrary.h"
 
+#include "LevelUtils.h"
+
 DEFINE_LOG_CATEGORY_STATIC(LogSMInstanceLevelEditorSelection, Log, All);
 
 bool FSMInstanceElementLevelEditorSelectionCustomization::CanSelectElement(const TTypedElement<ITypedElementSelectionInterface>& InElementSelectionHandle, const FTypedElementSelectionOptions& InSelectionOptions)
@@ -16,6 +18,14 @@ bool FSMInstanceElementLevelEditorSelectionCustomization::CanSelectElement(const
 	{
 		return false;
 	}
+
+	AActor* Owner = SMInstance.GetISMComponent()->GetOwner();
+	AActor* SelectionRoot = Owner->GetRootSelectionParent();
+	ULevel* SelectionLevel = (SelectionRoot != nullptr) ? SelectionRoot->GetLevel() : Owner->GetLevel();
+	if (!Owner->IsTemplate() && FLevelUtils::IsLevelLocked(SelectionLevel))
+	{
+		return false;
+	}
 	
 	// Bail if global selection is locked
 	return !GEdSelectionLock;
diff --git a/Engine/Source/Editor/WorldBrowser/Private/LevelModel.cpp b/Engine/Source/Editor/WorldBrowser/Private/LevelModel.cpp
--- a/Engine/Source/Editor/WorldBrowser/Private/LevelModel.cpp
+++ b/Engine/Source/Editor/WorldBrowser/Private/LevelModel.cpp
@@ -14,6 +14,10 @@
 #include "LevelUtils.h"
 #include "EditorLevelUtils.h"
 #include "ActorEditorUtils.h"
+#include "LevelEditor.h"
+#include "Elements/Framework/TypedElementSelectionSet.h"
+#include "Elements/Framework/TypedElementCommonActions.h"
+#include "Elements/Interfaces/TypedElementDetailsInterface.h"
 
 #include "Engine/LevelScriptBlueprint.h"
 
@@ -337,9 +341,29 @@ void FLevelModel::SetLocked(bool bLocked)
 	// If locking the level, deselect all of its actors and BSP surfaces
 	if (bLocked)
 	{
-		DeselectAllActors();
 		DeselectAllSurfaces();
 
+		FLevelEditorModule& LevelEditorModule = FModuleManager::Get().LoadModuleChecked<FLevelEditorModule>("LevelEditor");
+		TSharedPtr<ILevelEditor> Editor = LevelEditorModule.GetFirstLevelEditor();
+		UTypedElementSelectionSet* SelectionSet = Editor->GetMutableElementSelectionSet();
+		TArray<FTypedElementHandle> LevelElementHandles;
+
+		LevelElementHandles.Reserve(Level->Actors.Num());
+
+		// filter out elements that don't belong to the current level
+		SelectionSet->ForEachSelectedElement<ITypedElementWorldInterface>(
+			[Level, &LevelElementHandles](const TTypedElement<ITypedElementWorldInterface>& Element)
+			{
+				if (Element.GetOwnerLevel() == Level)
+				{
+					LevelElementHandles.Add(Element);
+				}
+				return true;
+			});
+
+		// deselect all the elements in the current level
+		SelectionSet->DeselectElements(LevelElementHandles, FTypedElementSelectionOptions());
+
 		// Tell the editor selection status was changed.
 		GEditor->NoteSelectionChange();
 
diff --git a/Engine/Source/Editor/WorldBrowser/WorldBrowser.Build.cs b/Engine/Source/Editor/WorldBrowser/WorldBrowser.Build.cs
--- a/Engine/Source/Editor/WorldBrowser/WorldBrowser.Build.cs
+++ b/Engine/Source/Editor/WorldBrowser/WorldBrowser.Build.cs
@@ -49,7 +49,9 @@ public class WorldBrowser : ModuleRules
                 "RHI",
                 "Json",
 				"ToolMenus",
-            }
+				"TypedElementRuntime",
+				"TypedElementFramework",
+			}
 		);
 
         DynamicallyLoadedModuleNames.AddRange(
