Don't allow adding partitioned worlds as level streaming.

diff --git a/Engine/Source/Editor/WorldBrowser/Private/StreamingLevels/StreamingLevelCollectionModel.cpp b/Engine/Source/Editor/WorldBrowser/Private/StreamingLevels/StreamingLevelCollectionModel.cpp
--- a/Engine/Source/Editor/WorldBrowser/Private/StreamingLevels/StreamingLevelCollectionModel.cpp
+++ b/Engine/Source/Editor/WorldBrowser/Private/StreamingLevels/StreamingLevelCollectionModel.cpp
@@ -461,21 +461,36 @@ void FStreamingLevelCollectionModel::HandleAddExistingLevelSelected(const TArray
 	TArray<FString> PackageNames;
 	for (const FAssetData& AssetData : SelectedAssets)
 	{
-		PackageNames.Add(AssetData.PackageName.ToString());
+		if (ULevel::GetIsLevelPartitionedFromPackage(AssetData.PackageName))
+		{
+			const FText MessageText = FText::Format(NSLOCTEXT("UnrealEd", "LevelIsPartitioned", "Level ({0}) is partitioned, can't add."), FText::FromName(AssetData.PackageName));
+
+			FSuppressableWarningDialog::FSetupInfo Info(MessageText, LOCTEXT("AddLevelToWorld_Title", "Add Level"), "LevelIsPartitionedWarning");
+			Info.ConfirmText = LOCTEXT("IsPartitioned_Ok", "Ok");
+			FSuppressableWarningDialog RemoveLevelWarning(Info);
+			RemoveLevelWarning.ShowModal();
+		}
+		else
+		{
+			PackageNames.Add(AssetData.PackageName.ToString());
+		}
 	}
 
-	// Save or selected list, adding a new level will clean it up
-	FLevelModelList SavedInvalidSelectedLevels = InvalidSelectedLevels;
+	if (PackageNames.Num())
+	{
+		// Save or selected list, adding a new level will clean it up
+		FLevelModelList SavedInvalidSelectedLevels = InvalidSelectedLevels;
 
-	EditorLevelUtils::AddLevelsToWorld(CurrentWorld.Get(), MoveTemp(PackageNames), AddedLevelStreamingClass);
+		EditorLevelUtils::AddLevelsToWorld(CurrentWorld.Get(), MoveTemp(PackageNames), AddedLevelStreamingClass);
 
-	// Force a cached level list rebuild
-	PopulateLevelsList();
+		// Force a cached level list rebuild
+		PopulateLevelsList();
 
-	if (bRemoveInvalidSelectedLevelsAfter)
-	{
-		InvalidSelectedLevels = SavedInvalidSelectedLevels;
-		RemoveInvalidSelectedLevels_Executed();
+		if (bRemoveInvalidSelectedLevelsAfter)
+		{
+			InvalidSelectedLevels = SavedInvalidSelectedLevels;
+			RemoveInvalidSelectedLevels_Executed();
+		}
 	}
 }
 
