Check for valid LevelModel in FLevelModelTreeItem::SetExpansion.

diff --git a/Engine/Source/Editor/WorldBrowser/Private/WorldTreeItemTypes.cpp b/Engine/Source/Editor/WorldBrowser/Private/WorldTreeItemTypes.cpp
--- a/Engine/Source/Editor/WorldBrowser/Private/WorldTreeItemTypes.cpp
+++ b/Engine/Source/Editor/WorldBrowser/Private/WorldTreeItemTypes.cpp
@@ -183,7 +183,11 @@ namespace WorldHierarchy
 	void FLevelModelTreeItem::SetExpansion(bool bExpanded)
 	{
 		IWorldTreeItem::SetExpansion(bExpanded);
-		LevelModel.Pin()->SetLevelExpansionFlag(bExpanded);
+
+		if (LevelModel.IsValid())
+		{
+			LevelModel.Pin()->SetLevelExpansionFlag(bExpanded);
+		}
 	}
 
 	bool FLevelModelTreeItem::HasModel(TSharedPtr<FLevelModel> InLevelModel) const
