Allow user to toggle visibily of columns in the world browser widget, displayed in the Levels tab. This also exposes a public API to allow programmetically changing the column visiblity.
- SWorldHierarchyImpl now shows the SHeaderRow (previously it was hidden)

diff --git a/Engine/Source/Editor/LevelEditor/Private/SLevelEditor.cpp b/Engine/Source/Editor/LevelEditor/Private/SLevelEditor.cpp
--- a/Engine/Source/Editor/LevelEditor/Private/SLevelEditor.cpp
+++ b/Engine/Source/Editor/LevelEditor/Private/SLevelEditor.cpp
@@ -77,6 +77,7 @@
 #include "LevelViewportTabContent.h"
 #include "SLevelViewport.h"
 #include "LevelEditorOutlinerSettings.h"
+#include "IWorldHierarchy.h"
 
 #define LOCTEXT_NAMESPACE "SLevelEditor"
 
@@ -1058,10 +1059,12 @@ TSharedRef<SDockTab> SLevelEditor::SpawnLevelEditorTab( const FSpawnTabArgs& Arg
 	else if( TabIdentifier == LevelEditorTabIds::WorldBrowserHierarchy)
 	{
 		FWorldBrowserModule& WorldBrowserModule = FModuleManager::LoadModuleChecked<FWorldBrowserModule>( "WorldBrowser" );
+		const TSharedRef<UE::WorldHierarchy::IWorldHierarchy> LevelHierarchyWidget = WorldBrowserModule.CreateWorldBrowserHierarchyWidget();
+		WorldHierarchy = LevelHierarchyWidget;
 		return SNew( SDockTab )
 			.Label( NSLOCTEXT("LevelEditor", "WorldBrowserHierarchyTabTitle", "Levels") )
 			[
-				WorldBrowserModule.CreateWorldBrowserHierarchy()
+				LevelHierarchyWidget->GetWidget()
 			];
 	}
 	else if( TabIdentifier == LevelEditorTabIds::WorldBrowserDetails)
diff --git a/Engine/Source/Editor/LevelEditor/Private/SLevelEditor.h b/Engine/Source/Editor/LevelEditor/Private/SLevelEditor.h
--- a/Engine/Source/Editor/LevelEditor/Private/SLevelEditor.h
+++ b/Engine/Source/Editor/LevelEditor/Private/SLevelEditor.h
@@ -26,6 +26,7 @@ class SDockTab;
 class SLevelEditorModeContent;
 class SLevelEditorToolBox;
 class UTypedElementSelectionSet;
+namespace UE::WorldHierarchy { class IWorldHierarchy; }
 
 /**
  * Unreal editor level editor Slate widget
@@ -165,6 +166,8 @@ public:
 	/** Return the most recently interacted with Outliner */
 	UE_DEPRECATED(5.1, "The Level Editor has multiple outliners, use GetAllSceneOutliners() or GetMostRecentlyUsedSceneOutliner() instead to avoid ambiguity")
 	virtual TSharedPtr<ISceneOutliner> GetSceneOutliner() const override;
+
+	virtual TWeakPtr<UE::WorldHierarchy::IWorldHierarchy> GetWorldHierarchy() override { return WorldHierarchy; }
 	
 	TSharedRef<SWidget> GetTitleBarMessageWidget() const { return TtileBarMessageBox.ToSharedRef(); }
 private:
@@ -339,6 +342,9 @@ private:
 	/** Weak pointer to the level editor's most recently created scene outliner */
 	TWeakPtr<ISceneOutliner> SceneOutlinerPtr;
 
+	/** The content displayed by the LevelEditorTabIds::WorldBrowserHierarchy tab. */
+	TWeakPtr<UE::WorldHierarchy::IWorldHierarchy> WorldHierarchy;
+
 	/** Map containing Weak pointers to all the Outliners in the level editor */
 	TMap<FName, TWeakPtr<ISceneOutliner>> SceneOutliners;
 
diff --git a/Engine/Source/Editor/LevelEditor/Public/ILevelEditor.h b/Engine/Source/Editor/LevelEditor/Public/ILevelEditor.h
--- a/Engine/Source/Editor/LevelEditor/Public/ILevelEditor.h
+++ b/Engine/Source/Editor/LevelEditor/Public/ILevelEditor.h
@@ -20,6 +20,7 @@ class ISceneOutliner;
 class IAssetViewport;
 class SLevelViewport;
 class UTypedElementSelectionSet;
+namespace UE::WorldHierarchy { class IWorldHierarchy; }
 
 /**
  * Public interface to SLevelEditor
@@ -106,6 +107,9 @@ public:
 
 	/** Return the most recently interacted with Outliner */
 	virtual TSharedPtr<ISceneOutliner> GetMostRecentlyUsedSceneOutliner() = 0;
+
+	/** @return The widget displayed in the Levels tab, if the tab is summoned. */
+	virtual TWeakPtr<UE::WorldHierarchy::IWorldHierarchy> GetWorldHierarchy() = 0;
 };
 
 
diff --git a/Engine/Source/Editor/WorldBrowser/Private/SWorldHierarchy.cpp b/Engine/Source/Editor/WorldBrowser/Private/SWorldHierarchy.cpp
--- a/Engine/Source/Editor/WorldBrowser/Private/SWorldHierarchy.cpp
+++ b/Engine/Source/Editor/WorldBrowser/Private/SWorldHierarchy.cpp
@@ -49,6 +49,24 @@ void SWorldHierarchy::Construct(const FArguments& InArgs)
 	OnBrowseWorld(InArgs._InWorld);
 }
 
+bool SWorldHierarchy::IsColumnVisible(FName Column) const
+{
+	const TSharedPtr<SWorldHierarchyImpl> WidgetPin = WeakWorldHierarchyImpl.Pin();
+	// This widget is destroyed during world construction
+	return ensureMsgf(WidgetPin, TEXT("Did you call this during world destruction?"))
+		&& WidgetPin->IsColumnVisible(Column);
+}
+
+void SWorldHierarchy::SetColumnVisible(FName Column, bool bVisible) const
+{
+	if (const TSharedPtr<SWorldHierarchyImpl> WidgetPin = WeakWorldHierarchyImpl.Pin();
+		// This widget is destroyed during world construction
+		ensureMsgf(WidgetPin, TEXT("Did you call this during world destruction?")))
+	{
+		WidgetPin->SetColumnVisible(Column, bVisible);
+	}
+}
+
 void SWorldHierarchy::OnBrowseWorld(UWorld* InWorld)
 {
 	// Remove all binding to an old world
diff --git a/Engine/Source/Editor/WorldBrowser/Private/SWorldHierarchy.h b/Engine/Source/Editor/WorldBrowser/Private/SWorldHierarchy.h
--- a/Engine/Source/Editor/WorldBrowser/Private/SWorldHierarchy.h
+++ b/Engine/Source/Editor/WorldBrowser/Private/SWorldHierarchy.h
@@ -2,6 +2,7 @@
 #pragma once
 
 #include "CoreMinimal.h"
+#include "IWorldHierarchy.h"
 #include "Layout/Visibility.h"
 #include "Input/Reply.h"
 #include "Widgets/SWidget.h"
@@ -9,11 +10,13 @@
 #include "Widgets/SCompoundWidget.h"
 
 class FLevelCollectionModel;
+class SWorldHierarchyImpl;
 struct FSlateBrush;
 
-
+/** Listens for world changes and updates SWorldHierarchyImpl, which displays the level hierarchy for the passed in world. */
 class SWorldHierarchy
 	: public SCompoundWidget
+	, public UE::WorldHierarchy::IWorldHierarchy
 {
 	SLATE_BEGIN_ARGS(SWorldHierarchy)
 		:_InWorld(nullptr)
@@ -25,6 +28,12 @@ class SWorldHierarchy
 	SWorldHierarchy();
 	~SWorldHierarchy();
 
+	//~ Begin IWorldHierarchy Interface
+	virtual TSharedRef<SWidget> GetWidget() override { return SharedThis(this); }
+	virtual bool IsColumnVisible(FName Column) const override;
+	virtual void SetColumnVisible(FName Column, bool bVisible) const override;
+	//~ End IWorldHierarchy Interface
+
 private:
 	void OnBrowseWorld(UWorld* InWorld);
 
@@ -39,5 +48,18 @@ private:
 	TSharedRef<SWidget> GetFileButtonContent();
 
 private:
+	
+	/**
+	 * Model for the UI managing the world logic.
+	 * FWorldBrowserModule expects to be the only referencer of this model when switching worlds.
+	 */
 	TSharedPtr<FLevelCollectionModel> WorldModel;
+	/**
+	 * Actually displays the hierarchy.
+	 *
+	 * WorldModel, which is managed by FWorldBrowserModule, expects the reference count to be 1 when switching worlds.
+	 * OnBrowseWorld kills this widget when switching worlds.
+	 * That is why this must be a weak pointer instead of a strong pointer.
+	 */
+	TWeakPtr<SWorldHierarchyImpl> WeakWorldHierarchyImpl;
 };
diff --git a/Engine/Source/Editor/WorldBrowser/Private/SWorldHierarchyImpl.cpp b/Engine/Source/Editor/WorldBrowser/Private/SWorldHierarchyImpl.cpp
--- a/Engine/Source/Editor/WorldBrowser/Private/SWorldHierarchyImpl.cpp
+++ b/Engine/Source/Editor/WorldBrowser/Private/SWorldHierarchyImpl.cpp
@@ -25,6 +25,8 @@
 #include "LevelFolders.h"
 #include "ScopedTransaction.h"
 #include "Editor.h"
+#include "WorldHierarchyColumns.h"
+#include "RevisionControlStyle/RevisionControlStyle.h"
 
 #define LOCTEXT_NAMESPACE "WorldBrowser"
 
@@ -98,93 +100,7 @@ void SWorldHierarchyImpl::Construct(const FArguments& InArgs)
 		SearchBoxLevelFilter->OnChanged().AddSP(this, &SWorldHierarchyImpl::FullRefresh);
 	}
 	SearchBoxHierarchyFilter->OnChanged().AddSP(this, &SWorldHierarchyImpl::FullRefresh);
-
-	HeaderRowWidget =
-		SNew( SHeaderRow )
-		.Visibility(EVisibility::Collapsed)
-
-		/** Level visibility column */
-		+ SHeaderRow::Column(HierarchyColumns::ColumnID_Visibility)
-		.Visibility(bFoldersOnlyMode ? EVisibility::Collapsed : EVisibility::Visible)
-		.FixedWidth(24.0f)
-		.HeaderContent()
-		[
-			SNew(STextBlock)
-			.ToolTipText(NSLOCTEXT("WorldBrowser", "Visibility", "Visibility"))
-		]
-
-		/** LevelName label column */
-		+ SHeaderRow::Column( HierarchyColumns::ColumnID_LevelLabel )
-			.FillWidth( 0.45f )
-			.HeaderContent()
-			[
-				SNew(STextBlock)
-					.ToolTipText(LOCTEXT("Column_LevelNameLabel", "Level"))
-					
-			]
-
-		/** Lighting Scenario column */
-		+ SHeaderRow::Column(HierarchyColumns::ColumnID_LightingScenario)
-			.Visibility(bFoldersOnlyMode ? EVisibility::Collapsed : EVisibility::Visible)
-			.FixedWidth( 24.0f )
-			.HeaderContent()
-			[
-				SNew(STextBlock)
-					.ToolTipText(NSLOCTEXT("WorldBrowser", "Lighting Scenario", "Lighting Scenario"))
-			]
-	
-
-		/** Level lock column */
-		+SHeaderRow::Column(HierarchyColumns::ColumnID_Lock)
-			.Visibility(bFoldersOnlyMode ? EVisibility::Collapsed : EVisibility::Visible)
-			.FixedWidth( 24.0f )
-			.HeaderContent()
-			[
-				SNew(STextBlock)
-					.ToolTipText(NSLOCTEXT("WorldBrowser", "Lock", "Lock"))
-			]
 	
-		/** Level kismet column */
-		+ SHeaderRow::Column(HierarchyColumns::ColumnID_Kismet)
-			.Visibility(bFoldersOnlyMode ? EVisibility::Collapsed : EVisibility::Visible)
-			.FixedWidth( 24.0f )
-			.HeaderContent()
-			[
-				SNew(STextBlock)
-					.ToolTipText(NSLOCTEXT("WorldBrowser", "Blueprint", "Open the level blueprint for this Level"))
-			]
-
-		/** Level SCC status column */
-		+ SHeaderRow::Column(HierarchyColumns::ColumnID_SCCStatus)
-			.Visibility(bFoldersOnlyMode ? EVisibility::Collapsed : EVisibility::Visible)
-			.FixedWidth( 24.0f )
-			.HeaderContent()
-			[
-				SNew(STextBlock)
-					.ToolTipText(NSLOCTEXT("WorldBrowser", "SCCStatus", "Status in Revision Control"))
-			]
-
-		/** Level save column */
-		+ SHeaderRow::Column(HierarchyColumns::ColumnID_Save)
-			.Visibility(bFoldersOnlyMode ? EVisibility::Collapsed : EVisibility::Visible)
-			.FixedWidth( 24.0f )
-			.HeaderContent()
-			[
-				SNew(STextBlock)
-					.ToolTipText(NSLOCTEXT("WorldBrowser", "Save", "Save this Level"))
-			]
-
-		/** Level color column */
-		+ SHeaderRow::Column(HierarchyColumns::ColumnID_Color)
-			.Visibility(bFoldersOnlyMode ? EVisibility::Collapsed : EVisibility::Visible)
-			.FixedWidth(24.0f)
-			.HeaderContent()
-			[
-				SNew(STextBlock)
-				.ToolTipText(NSLOCTEXT("WorldBrowser", "Color", "Color used for visualization of Level"))
-			];
-
-
 	FOnContextMenuOpening ContextMenuEvent;
 	if (!bFoldersOnlyMode)
 	{
@@ -271,7 +187,7 @@ void SWorldHierarchyImpl::Construct(const FArguments& InArgs)
 				.OnMouseButtonDoubleClick(this, &SWorldHierarchyImpl::OnTreeViewMouseButtonDoubleClick)
 				.OnContextMenuOpening(ContextMenuEvent)
 				.OnItemScrolledIntoView(this, &SWorldHierarchyImpl::OnTreeItemScrolledIntoView)
-				.HeaderRow(HeaderRowWidget.ToSharedRef())
+				.HeaderRow(CreateHeaderRow())
 			]
 		]
 
@@ -367,6 +283,117 @@ void SWorldHierarchyImpl::RefreshView()
 	bNeedsRefresh = true;
 }
 
+TSharedRef<SHeaderRow> SWorldHierarchyImpl::CreateHeaderRow()
+{
+	using namespace UE::WorldHierarchy;
+	constexpr float IconWidth = 24.f;
+	return SAssignNew(HeaderRowWidget, SHeaderRow)
+		.CanSelectGeneratedColumn(true) // Lets the user choose which columns should be shown
+
+		/** Level visibility column */
+		+ SHeaderRow::Column(HierarchyColumns::ColumnID_EditorVisibility)
+		.Visibility(bFoldersOnlyMode ? EVisibility::Collapsed : EVisibility::Visible)
+		.DefaultLabel(LOCTEXT("Column.Visibility.DefaultLabel", "Visibility"))
+		.FixedWidth(IconWidth)
+		.HeaderContent()
+		[
+			SNew(SImage)
+			.Image(FAppStyle::Get().GetBrush("Level.VisibleIcon16x"))
+			.ColorAndOpacity(FSlateColor::UseForeground())
+			.ToolTipText(LOCTEXT("Column.Visibility.Tooltip", "Toggles the editor visibility"))
+		]
+
+		/** LevelName label column */
+		+ SHeaderRow::Column( HierarchyColumns::ColumnID_LevelLabel )
+		.FillWidth(1.f)
+		.DefaultLabel(LOCTEXT("Column.Level.DefaultLabel", "Level"))
+		.ShouldGenerateWidget(true) // This column cannot be toggled off by the user
+		.HeaderContent()
+		[
+			SNew(STextBlock)
+			.Text(LOCTEXT("Column.Level.Label", "Level"))
+			.ToolTipText(LOCTEXT("Column.Level.Tooltip", "Level"))
+		]
+
+		/** Lighting Scenario column */
+		+ SHeaderRow::Column(HierarchyColumns::ColumnID_LightingScenario)
+		.Visibility(bFoldersOnlyMode ? EVisibility::Collapsed : EVisibility::Visible)
+		.DefaultLabel(LOCTEXT("Column.LightingScenario.DefaultLabel", "Lighting Scenario"))
+		.FixedWidth(IconWidth)
+		.HeaderContent()
+		[
+			SNew(SImage)
+			.Image(FAppStyle::Get().GetBrush("Level.LightingScenarioIconSolid16x"))
+			.ColorAndOpacity(FSlateColor::UseForeground())
+			.ToolTipText(LOCTEXT("Column.LightingScenario.Tooltip", "Lighting Scenario"))
+		]
+
+		/** Level lock column */
+		+SHeaderRow::Column(HierarchyColumns::ColumnID_Lock)
+		.Visibility(bFoldersOnlyMode ? EVisibility::Collapsed : EVisibility::Visible)
+		.DefaultLabel(LOCTEXT("Column.LevelLock.DefaultLabel", "Lock"))
+		.FixedWidth(IconWidth)
+		.HeaderContent()
+		[
+			SNew(SImage)
+			.Image(FAppStyle::Get().GetBrush("Icons.Lock"))
+			.ColorAndOpacity(FSlateColor::UseForeground())
+			.ToolTipText(LOCTEXT("Column.LevelLock.Tooltip", "Lock"))
+		]
+	
+		/** Level kismet column */
+		+ SHeaderRow::Column(HierarchyColumns::ColumnID_Kismet)
+		.Visibility(bFoldersOnlyMode ? EVisibility::Collapsed : EVisibility::Visible)
+		.DefaultLabel(LOCTEXT("Column.Kismet.DefaultLabel", "Open Blueprint"))
+		.FixedWidth(IconWidth)
+		.HeaderContent()
+		[
+			SNew(SImage)
+			.ColorAndOpacity(FSlateColor::UseForeground())
+			.Image(FAppStyle::Get().GetBrush("Icons.Blueprints"))
+			.ToolTipText(LOCTEXT("Column.Kismet.Tooltip", "Open the level blueprint for this Level"))
+		]
+
+		/** Level SCC status column */
+		+ SHeaderRow::Column(HierarchyColumns::ColumnID_SCCStatus)
+		.Visibility(bFoldersOnlyMode ? EVisibility::Collapsed : EVisibility::Visible)
+		.DefaultLabel(LOCTEXT("Column.SCC.DefaultLabel", "Revision Control"))
+		.FixedWidth(IconWidth)
+		.HeaderContent()
+		[
+			SNew(SImage)
+			.ColorAndOpacity(FSlateColor::UseForeground())
+			.Image(FRevisionControlStyleManager::Get().GetBrush("RevisionControl.Icon"))
+			.ToolTipText(LOCTEXT("Column.SCC.Tooltip", "Status in Revision Control"))
+		]
+
+		/** Level save column */
+		+ SHeaderRow::Column(HierarchyColumns::ColumnID_Save)
+		.Visibility(bFoldersOnlyMode ? EVisibility::Collapsed : EVisibility::Visible)
+		.DefaultLabel(LOCTEXT("Column.Save.DefaultLabel", "Save"))
+		.FixedWidth(IconWidth)
+		.HeaderContent()
+		[
+			SNew(SImage)
+			.ColorAndOpacity(FSlateColor::UseForeground())
+			.Image(FAppStyle::Get().GetBrush("Icons.SaveModified"))
+			.ToolTipText(LOCTEXT("Column.Save.Tooltip", "Save this Level"))
+		]
+
+		/** Level color column */
+		+ SHeaderRow::Column(HierarchyColumns::ColumnID_Color)
+		.Visibility(bFoldersOnlyMode ? EVisibility::Collapsed : EVisibility::Visible)
+		.DefaultLabel(LOCTEXT("Column.Color.DefaultLabel", "Color"))
+		.FixedWidth(IconWidth)
+		.HeaderContent()
+		[
+			SNew(SImage)
+			.ColorAndOpacity(FSlateColor(FColor::White))
+			.Image(FAppStyle::Get().GetBrush("Level.ColorIcon"))
+			.ToolTipText(LOCTEXT("Column.Color.Tooltip", "Color used for visualization of Level"))
+		];
+}
+
 TSharedRef<ITableRow> SWorldHierarchyImpl::GenerateTreeRow(WorldHierarchy::FWorldTreeItemPtr Item, const TSharedRef<STableViewBase>& OwnerTable)
 {
 	check(Item.IsValid());
@@ -692,6 +719,16 @@ void SWorldHierarchyImpl::AddDroppedLevelsToFolder(const TArray<FAssetData>& Wor
 	}
 }
 
+bool SWorldHierarchyImpl::IsColumnVisible(FName ColumnId) const
+{
+	return HeaderRowWidget->IsColumnVisible(ColumnId);
+}
+
+void SWorldHierarchyImpl::SetColumnVisible(FName ColumnId, bool bVisible) const
+{
+	HeaderRowWidget->SetShowGeneratedColumn(ColumnId, bVisible);
+}
+
 void SWorldHierarchyImpl::OnTreeItemScrolledIntoView( WorldHierarchy::FWorldTreeItemPtr Item, const TSharedPtr<ITableRow>& Widget )
 {
 	if (Item == ItemPendingRename)
diff --git a/Engine/Source/Editor/WorldBrowser/Private/SWorldHierarchyImpl.h b/Engine/Source/Editor/WorldBrowser/Private/SWorldHierarchyImpl.h
--- a/Engine/Source/Editor/WorldBrowser/Private/SWorldHierarchyImpl.h
+++ b/Engine/Source/Editor/WorldBrowser/Private/SWorldHierarchyImpl.h
@@ -69,7 +69,6 @@ public:
 	SLATE_END_ARGS()
 
 	SWorldHierarchyImpl();
-
 	~SWorldHierarchyImpl();
 	
 	void Construct(const FArguments& InArgs);
@@ -100,6 +99,17 @@ public:
 	/** Helper funciton to get the selected items from the tree widget */
 	TArray<WorldHierarchy::FWorldTreeItemPtr> GetSelectedTreeItems() const { return TreeWidget->GetSelectedItems(); }
 
+	/**
+	 * @see WorldHierarchyColumns.h for named columns.
+	 * @return Whether Column is visible in the UI.
+	 */
+	bool IsColumnVisible(FName ColumnId) const;
+	/**
+	 * Sets whether Column is visible in the UI.
+	 * @see WorldHierarchyColumns.h for named columns.
+	 */
+	void SetColumnVisible(FName ColumnId, bool bVisible) const;
+
 public:
 	//~ FEditorUndoClient
 	virtual void PostUndo(bool bSuccess) override;
@@ -109,6 +119,9 @@ private:
 	// The maximum number of pending operations to process at one time
 	static const int32 MaxPendingOperations = 500;
 
+	/** Creates the header row for the level hierarchy. */
+	TSharedRef<SHeaderRow> CreateHeaderRow();
+	
 	/** Creates an item for the tree view */
 	TSharedRef<ITableRow> GenerateTreeRow(WorldHierarchy::FWorldTreeItemPtr Item, const TSharedRef<STableViewBase>& OwnerTable);
 
diff --git a/Engine/Source/Editor/WorldBrowser/Private/SWorldHierarchyItem.cpp b/Engine/Source/Editor/WorldBrowser/Private/SWorldHierarchyItem.cpp
--- a/Engine/Source/Editor/WorldBrowser/Private/SWorldHierarchyItem.cpp
+++ b/Engine/Source/Editor/WorldBrowser/Private/SWorldHierarchyItem.cpp
@@ -13,6 +13,7 @@
 #include "Editor.h"
 #include "LevelCollectionModel.h"
 #include "WorldBrowserDragDrop.h"
+#include "WorldHierarchyColumns.h"
 #include "SWorldHierarchyImpl.h"
 
 #include "Widgets/Views/SListView.h"
@@ -53,6 +54,7 @@ void SWorldHierarchyItem::Construct(const FArguments& InArgs, TSharedRef<STableV
 BEGIN_SLATE_FUNCTION_BUILD_OPTIMIZATION
 TSharedRef< SWidget > SWorldHierarchyItem::GenerateWidgetForColumn( const FName& ColumnID )
 {
+	using namespace UE::WorldHierarchy;
 	TSharedPtr< SWidget > TableRowContent = SNullWidget::NullWidget;
 
 	if (ColumnID == HierarchyColumns::ColumnID_LevelLabel)
@@ -163,7 +165,7 @@ TSharedRef< SWidget > SWorldHierarchyItem::GenerateWidgetForColumn( const FName&
 				]
 			;
 		}
-		else if (ColumnID == HierarchyColumns::ColumnID_Visibility)
+		else if (ColumnID == HierarchyColumns::ColumnID_EditorVisibility)
 		{
 			TableRowContent =
 				SAssignNew(VisibilityButton, SButton)
diff --git a/Engine/Source/Editor/WorldBrowser/Private/SWorldHierarchyItem.h b/Engine/Source/Editor/WorldBrowser/Private/SWorldHierarchyItem.h
--- a/Engine/Source/Editor/WorldBrowser/Private/SWorldHierarchyItem.h
+++ b/Engine/Source/Editor/WorldBrowser/Private/SWorldHierarchyItem.h
@@ -20,22 +20,6 @@ class FLevelModel;
 class SButton;
 class SWorldHierarchyImpl;
 
-namespace HierarchyColumns
-{
-	/** IDs for list columns */
-	static const FName ColumnID_LevelLabel( "Level" );
-	static const FName ColumnID_Visibility( "Visibility" );
-	static const FName ColumnID_LightingScenario( "LightingScenario" );
-	static const FName ColumnID_Lock( "Lock" );
-	static const FName ColumnID_SCCStatus( "SCC_Status" );
-	static const FName ColumnID_Save( "Save" );
-	static const FName ColumnID_Color("Color");
-	static const FName ColumnID_Kismet( "Blueprint" );
-	static const FName ColumnID_ActorCount( "ActorCount" );
-	static const FName ColumnID_LightmassSize( "LightmassSize" );
-	static const FName ColumnID_FileSize( "FileSize" );
-}
-
 /** A single item in the levels hierarchy tree. Represents a level model */
 class SWorldHierarchyItem 
 	: public SMultiColumnTableRow<WorldHierarchy::FWorldTreeItemPtr>
diff --git a/Engine/Source/Editor/WorldBrowser/Private/WorldBrowserModule.cpp b/Engine/Source/Editor/WorldBrowser/Private/WorldBrowserModule.cpp
--- a/Engine/Source/Editor/WorldBrowser/Private/WorldBrowserModule.cpp
+++ b/Engine/Source/Editor/WorldBrowser/Private/WorldBrowserModule.cpp
@@ -127,11 +127,16 @@ void FWorldBrowserModule::ReleaseWorldModel()
 }
 
 TSharedRef<SWidget> FWorldBrowserModule::CreateWorldBrowserHierarchy()
+{
+	return CreateWorldBrowserHierarchyWidget()->GetWidget();
+}
+
+TSharedRef<UE::WorldHierarchy::IWorldHierarchy> FWorldBrowserModule::CreateWorldBrowserHierarchyWidget()
 {
 	UWorld* EditorWorld = GEditor->GetEditorWorldContext().World();
 	return SNew(SWorldHierarchy).InWorld(EditorWorld);
 }
-	
+
 TSharedRef<SWidget> FWorldBrowserModule::CreateWorldBrowserDetails()
 {
 	UWorld* EditorWorld = GEditor->GetEditorWorldContext().World();
diff --git a/Engine/Source/Editor/WorldBrowser/Private/WorldHierarchyColumns.cpp b/Engine/Source/Editor/WorldBrowser/Private/WorldHierarchyColumns.cpp
new file mode 100644
--- /dev/null
+++ b/Engine/Source/Editor/WorldBrowser/Private/WorldHierarchyColumns.cpp
@@ -0,0 +1,15 @@
+﻿// Copyright Epic Games, Inc. All Rights Reserved.
+
+#include "WorldHierarchyColumns.h"
+
+namespace UE::WorldHierarchy::HierarchyColumns
+{
+	const FName ColumnID_LevelLabel( "Level" );
+	const FName ColumnID_EditorVisibility( "EditorVisibility" );
+	const FName ColumnID_LightingScenario( "LightingScenario" );
+	const FName ColumnID_Lock( "Lock" );
+	const FName ColumnID_SCCStatus( "SCC_Status" );
+	const FName ColumnID_Save( "Save" );
+	const FName ColumnID_Color("Color");
+	const FName ColumnID_Kismet( "Blueprint" );
+}
\ No newline at end of file
diff --git a/Engine/Source/Editor/WorldBrowser/Public/IWorldHierarchy.h b/Engine/Source/Editor/WorldBrowser/Public/IWorldHierarchy.h
new file mode 100644
--- /dev/null
+++ b/Engine/Source/Editor/WorldBrowser/Public/IWorldHierarchy.h
@@ -0,0 +1,33 @@
+﻿// Copyright Epic Games, Inc. All Rights Reserved.
+
+#pragma once
+
+#include "Templates/SharedPointerFwd.h"
+
+class FName;
+class SWidget;
+
+namespace UE::WorldHierarchy
+{
+	/** Displays levels for the world. */
+	class IWorldHierarchy
+	{
+	public:
+
+		/** @return Gets the widget that displays the levels */
+		virtual TSharedRef<SWidget> GetWidget() = 0;
+
+		/**
+		 * @see WorldHierarchyColumns.h for named columns.
+		 * @return Whether Column is visible in the UI.
+		 */
+		virtual bool IsColumnVisible(FName Column) const = 0;
+		/**
+		 * Sets whether Column is visible in the UI.
+		 * @see WorldHierarchyColumns.h for named columns.
+		 */
+		virtual void SetColumnVisible(FName Column, bool bVisible) const = 0;
+
+		virtual ~IWorldHierarchy() = default;
+	};
+}
\ No newline at end of file
diff --git a/Engine/Source/Editor/WorldBrowser/Public/WorldBrowserModule.h b/Engine/Source/Editor/WorldBrowser/Public/WorldBrowserModule.h
--- a/Engine/Source/Editor/WorldBrowser/Public/WorldBrowserModule.h
+++ b/Engine/Source/Editor/WorldBrowser/Public/WorldBrowserModule.h
@@ -7,6 +7,7 @@
 #include "LevelEditor.h"
 
 class FLevelCollectionModel;
+namespace UE::WorldHierarchy { class IWorldHierarchy; }
 
 /**
  * The module holding all of the UI related pieces for SubLevels management
@@ -29,7 +30,10 @@ public:
 	/**
 	 * Creates a levels hierarchy widget
 	 */
-	virtual TSharedRef<class SWidget> CreateWorldBrowserHierarchy();
+	UE_DEPRECATED(5.5, "Use CreateWorldBrowserHierarchyWidget instead.")
+	virtual TSharedRef<SWidget> CreateWorldBrowserHierarchy();
+	/** Creates a levels hierarchy widget */
+	virtual TSharedRef<UE::WorldHierarchy::IWorldHierarchy> CreateWorldBrowserHierarchyWidget();
 	
 	/**
 	 * Creates a levels details widget
diff --git a/Engine/Source/Editor/WorldBrowser/Public/WorldHierarchyColumns.h b/Engine/Source/Editor/WorldBrowser/Public/WorldHierarchyColumns.h
new file mode 100644
--- /dev/null
+++ b/Engine/Source/Editor/WorldBrowser/Public/WorldHierarchyColumns.h
@@ -0,0 +1,18 @@
+﻿// Copyright Epic Games, Inc. All Rights Reserved.
+
+#pragma once
+
+#include "UObject/NameTypes.h"
+
+namespace UE::WorldHierarchy::HierarchyColumns
+{
+	/** IDs for list columns */
+	WORLDBROWSER_API extern const FName ColumnID_LevelLabel;
+	WORLDBROWSER_API extern const FName ColumnID_EditorVisibility;
+	WORLDBROWSER_API extern const FName ColumnID_LightingScenario;
+	WORLDBROWSER_API extern const FName ColumnID_Lock;
+	WORLDBROWSER_API extern const FName ColumnID_SCCStatus;
+	WORLDBROWSER_API extern const FName ColumnID_Save;
+	WORLDBROWSER_API extern const FName ColumnID_Color;
+	WORLDBROWSER_API extern const FName ColumnID_Kismet;
+}
\ No newline at end of file
