Introduce gameplay visibility column to level editor (which is hidden by default).

diff --git a/Engine/Source/Editor/WorldBrowser/Private/LevelCollectionCommands.h b/Engine/Source/Editor/WorldBrowser/Private/LevelCollectionCommands.h
--- a/Engine/Source/Editor/WorldBrowser/Private/LevelCollectionCommands.h
+++ b/Engine/Source/Editor/WorldBrowser/Private/LevelCollectionCommands.h
@@ -93,13 +93,21 @@ public:
 		UI_COMMAND( MoveFoliageToSelected, "Move Selected Foliage to Level", "Moves the selected foliage instances to this level. Keeps cross-level references to original bases", EUserInterfaceActionType::Button, FInputChord() );
 		UI_COMMAND( SelectStreamingVolumes, "Select Streaming Volumes", "Selects the streaming volumes associated with the selected levels", EUserInterfaceActionType::Button, FInputChord() );
 
-		//visibility
-		UI_COMMAND( World_ShowSelectedLevels, "Show Selected", "Toggles selected levels to a visible state in the viewports", EUserInterfaceActionType::Button, FInputChord() );
-		UI_COMMAND( World_HideSelectedLevels, "Hide Selected", "Toggles selected levels to an invisible state in the viewports", EUserInterfaceActionType::Button, FInputChord() );
-		UI_COMMAND( World_ShowOnlySelectedLevels, "Show Only Selected", "Toggles the selected levels to a visible state; toggles all other levels to an invisible state", EUserInterfaceActionType::Button, FInputChord() );
-		UI_COMMAND( World_ShowAllButSelectedLevels, "Show All But Selected", "Toggles the selected levels to an invisible state; toggles all other levels to a visible state", EUserInterfaceActionType::Button, FInputChord());
-		UI_COMMAND( World_ShowAllLevels, "Show All", "Toggles all levels to a visible state in the viewports", EUserInterfaceActionType::Button, FInputChord());
-		UI_COMMAND( World_HideAllLevels, "Hide All", "Hides all levels to an invisible state in the viewports", EUserInterfaceActionType::Button, FInputChord() );
+		//editor visibility
+		UI_COMMAND( World_ShowInEditorSelectedLevels, "Show Selected", "Toggles selected levels to a visible state in the viewports for editor worlds", EUserInterfaceActionType::Button, FInputChord() );
+		UI_COMMAND( World_HideInEditorSelectedLevels, "Hide Selected", "Toggles selected levels to an invisible state in the viewports for editor worlds", EUserInterfaceActionType::Button, FInputChord() );
+		UI_COMMAND( World_ShowInEditorOnlySelectedLevels, "Show Only Selected", "Toggles the selected levels to a visible state for editor worlds; toggles all other levels to an invisible state", EUserInterfaceActionType::Button, FInputChord() );
+		UI_COMMAND( World_ShowInEditorAllButSelectedLevels, "Show All But Selected", "Toggles the selected levels to an invisible state for editor worlds; toggles all other levels to a visible state", EUserInterfaceActionType::Button, FInputChord());
+		UI_COMMAND( World_ShowInEditorAllLevels, "Show All", "Toggles all levels to a visible state in the viewports for editor worlds", EUserInterfaceActionType::Button, FInputChord());
+		UI_COMMAND( World_HideInEditorAllLevels, "Hide All", "Hides all levels to an invisible state in the viewports for editor worlds", EUserInterfaceActionType::Button, FInputChord() );
+
+		//game visibility
+		UI_COMMAND( World_ShowInGameSelectedLevels, "Show Selected", "Toggles selected levels to a visible state in the viewports for game worlds", EUserInterfaceActionType::Button, FInputChord() );
+		UI_COMMAND( World_HideInGameSelectedLevels, "Hide Selected", "Toggles selected levels to an invisible state in the viewports for game worlds", EUserInterfaceActionType::Button, FInputChord() );
+		UI_COMMAND( World_ShowInGameOnlySelectedLevels, "Show Only Selected", "Toggles the selected levels to a visible state for game worlds; toggles all other levels to an invisible state", EUserInterfaceActionType::Button, FInputChord() );
+		UI_COMMAND( World_ShowInGameAllButSelectedLevels, "Show All But Selected", "Toggles the selected levels to an invisible state for game worlds; toggles all other levels to a visible state", EUserInterfaceActionType::Button, FInputChord());
+		UI_COMMAND( World_ShowInGameAllLevels, "Show All", "Toggles all levels to a visible state in the viewports for game worlds", EUserInterfaceActionType::Button, FInputChord());
+		UI_COMMAND( World_HideInGameAllLevels, "Hide All", "Hides all levels to an invisible state in the viewports for game worlds", EUserInterfaceActionType::Button, FInputChord() );
 
 		//lock
 		UI_COMMAND( World_LockSelectedLevels, "Lock Selected", "Locks selected levels", EUserInterfaceActionType::Button, FInputChord() );
@@ -266,22 +274,41 @@ public:
 	TSharedPtr< FUICommandInfo > SelectStreamingVolumes;
 	
 	/** Makes selected Levels visible */
-	TSharedPtr< FUICommandInfo > World_ShowSelectedLevels;
+	TSharedPtr< FUICommandInfo > World_ShowInEditorSelectedLevels;
 
 	/** Makes selected Levels invisible */
-	TSharedPtr< FUICommandInfo > World_HideSelectedLevels;
+	TSharedPtr< FUICommandInfo > World_HideInEditorSelectedLevels;
 
 	/** Makes selected Levels visible; makes all others invisible */
-	TSharedPtr< FUICommandInfo > World_ShowOnlySelectedLevels;
+	TSharedPtr< FUICommandInfo > World_ShowInEditorOnlySelectedLevels;
 
 	/** Makes selected Levels invisible; makes all others visible */
-	TSharedPtr< FUICommandInfo > World_ShowAllButSelectedLevels;
+	TSharedPtr< FUICommandInfo > World_ShowInEditorAllButSelectedLevels;
 
 	/** Makes all Levels visible */
-	TSharedPtr< FUICommandInfo > World_ShowAllLevels;
+	TSharedPtr< FUICommandInfo > World_ShowInEditorAllLevels;
 
 	/** Makes all Levels invisible */
-	TSharedPtr< FUICommandInfo > World_HideAllLevels;
+	TSharedPtr< FUICommandInfo > World_HideInEditorAllLevels;
+
+	
+	/** Makes selected Levels visible in the editor */
+	TSharedPtr< FUICommandInfo > World_ShowInGameSelectedLevels;
+	
+	/** Makes selected Levels invisible in the editor */
+	TSharedPtr< FUICommandInfo > World_HideInGameSelectedLevels;
+	
+	/** Makes selected Levels visible; makes all others invisible in the editor */
+	TSharedPtr< FUICommandInfo > World_ShowInGameOnlySelectedLevels;
+	
+	/** Makes selected Levels invisible; makes all others visible in the editor */
+	TSharedPtr< FUICommandInfo > World_ShowInGameAllButSelectedLevels;
+	
+	/** Makes all Levels visible in the editor */
+	TSharedPtr< FUICommandInfo > World_ShowInGameAllLevels;
+	
+	/** Makes all Levels invisible in the editor */
+	TSharedPtr< FUICommandInfo > World_HideInGameAllLevels;
 	
 	
 	/** Locks selected levels */
diff --git a/Engine/Source/Editor/WorldBrowser/Private/LevelCollectionModel.cpp b/Engine/Source/Editor/WorldBrowser/Private/LevelCollectionModel.cpp
--- a/Engine/Source/Editor/WorldBrowser/Private/LevelCollectionModel.cpp
+++ b/Engine/Source/Editor/WorldBrowser/Private/LevelCollectionModel.cpp
@@ -36,6 +36,7 @@
 #include "Engine/WorldComposition.h"
 #include "Misc/ScopeExit.h"
 #include "LevelUtils.h"
+#include "WorldTreeItemTypes.h"
 
 #define LOCTEXT_NAMESPACE "WorldBrowser"
 
@@ -159,28 +160,43 @@ void FLevelCollectionModel::BindCommands()
 		FExecuteAction::CreateSP(this, &FLevelCollectionModel::ConvertLevelToExternalActors_Executed, false),
 		FCanExecuteAction::CreateSP(this, &FLevelCollectionModel::CanConvertAnyLevelToExternalActors, false));
 
-	//visibility
-	ActionList.MapAction( Commands.World_ShowSelectedLevels,
-		FExecuteAction::CreateSP( this, &FLevelCollectionModel::ShowSelectedLevels_Executed  ),
+	//editor visibility
+	ActionList.MapAction( Commands.World_ShowInEditorSelectedLevels,
+		FExecuteAction::CreateSP( this, &FLevelCollectionModel::ShowInEditorSelectedLevels_Executed  ),
 		FCanExecuteAction::CreateSP( this, &FLevelCollectionModel::AreAnySelectedLevelsLoaded ) );
-	
-	ActionList.MapAction( Commands.World_HideSelectedLevels,
-		FExecuteAction::CreateSP( this, &FLevelCollectionModel::HideSelectedLevels_Executed  ),
+	ActionList.MapAction( Commands.World_HideInEditorSelectedLevels,
+		FExecuteAction::CreateSP( this, &FLevelCollectionModel::HideInEditorSelectedLevels_Executed  ),
 		FCanExecuteAction::CreateSP( this, &FLevelCollectionModel::AreAnySelectedLevelsLoaded ) );
-	
-	ActionList.MapAction( Commands.World_ShowOnlySelectedLevels,
-		FExecuteAction::CreateSP( this, &FLevelCollectionModel::ShowOnlySelectedLevels_Executed  ),
+	ActionList.MapAction( Commands.World_ShowInEditorOnlySelectedLevels,
+		FExecuteAction::CreateSP( this, &FLevelCollectionModel::ShowInEditorOnlySelectedLevels_Executed  ),
 		FCanExecuteAction::CreateSP( this, &FLevelCollectionModel::AreAnySelectedLevelsLoaded ) );
-
-	ActionList.MapAction(Commands.World_ShowAllButSelectedLevels,
-		FExecuteAction::CreateSP(this, &FLevelCollectionModel::ShowAllButSelectedLevels_Executed),
+	ActionList.MapAction(Commands.World_ShowInEditorAllButSelectedLevels,
+		FExecuteAction::CreateSP(this, &FLevelCollectionModel::ShowInEditorAllButSelectedLevels_Executed),
 		FCanExecuteAction::CreateSP(this, &FLevelCollectionModel::AreAnySelectedLevelsLoaded));
-
-	ActionList.MapAction(Commands.World_ShowAllLevels,
-		FExecuteAction::CreateSP(this, &FLevelCollectionModel::ShowAllLevels_Executed));
+	ActionList.MapAction(Commands.World_ShowInEditorAllLevels,
+		FExecuteAction::CreateSP(this, &FLevelCollectionModel::ShowInEditorAllLevels_Executed));
+	ActionList.MapAction(Commands.World_HideInEditorAllLevels,
+		FExecuteAction::CreateSP(this, &FLevelCollectionModel::HideInEditorAllLevels_Executed));
 	
-	ActionList.MapAction(Commands.World_HideAllLevels,
-		FExecuteAction::CreateSP(this, &FLevelCollectionModel::HideAllLevels_Executed));
+	//game visibility
+	ActionList.MapAction( Commands.World_ShowInGameSelectedLevels,
+		FExecuteAction::CreateSP( this, &FLevelCollectionModel::ShowInGameSelectedLevels_Executed  ),
+		FCanExecuteAction::CreateSP( this, &FLevelCollectionModel::CanExecuteGameVisibilityCommandsForSelectedLevels ) );
+	ActionList.MapAction( Commands.World_HideInGameSelectedLevels,
+		FExecuteAction::CreateSP( this, &FLevelCollectionModel::HideInGameSelectedLevels_Executed  ),
+		FCanExecuteAction::CreateSP( this, &FLevelCollectionModel::CanExecuteGameVisibilityCommandsForSelectedLevels ) );
+	ActionList.MapAction( Commands.World_ShowInGameOnlySelectedLevels,
+		FExecuteAction::CreateSP( this, &FLevelCollectionModel::ShowInGameOnlySelectedLevels_Executed  ),
+		FCanExecuteAction::CreateSP( this, &FLevelCollectionModel::CanExecuteGameVisibilityCommandsForSelectedLevels ) );
+	ActionList.MapAction(Commands.World_ShowInGameAllButSelectedLevels,
+		FExecuteAction::CreateSP(this, &FLevelCollectionModel::ShowInGameAllButSelectedLevels_Executed),
+		FCanExecuteAction::CreateSP(this, &FLevelCollectionModel::CanExecuteGameVisibilityCommandsForSelectedLevels));
+	ActionList.MapAction(Commands.World_ShowInGameAllLevels,
+		FExecuteAction::CreateSP(this, &FLevelCollectionModel::ShowInGameAllLevels_Executed),
+		FCanExecuteAction::CreateSP(this, &FLevelCollectionModel::CanExecuteGameVisibilityCommands));
+	ActionList.MapAction(Commands.World_HideInGameAllLevels,
+		FExecuteAction::CreateSP(this, &FLevelCollectionModel::HideInGameAllLevels_Executed),
+		FCanExecuteAction::CreateSP(this, &FLevelCollectionModel::CanExecuteGameVisibilityCommands));
 		
 	//lock
 	ActionList.MapAction( Commands.World_LockSelectedLevels,
@@ -463,7 +479,7 @@ void FLevelCollectionModel::IterateHierarchy(FLevelModelVisitor& Visitor)
 	}
 }
 
-void FLevelCollectionModel::HideLevels(const FLevelModelList& InLevelList)
+void FLevelCollectionModel::HideLevelsInEditor(const FLevelModelList& InLevelList)
 {
 	if (IsReadOnly())
 	{
@@ -493,12 +509,12 @@ void FLevelCollectionModel::HideLevels(const FLevelModelList& InLevelList)
 		LevelModels.Add(It->Get());
 		bVisible.Add(false);
 	}
-	FLevelModel::SetVisible(LevelModels, bVisible);
+	FLevelModel::SetVisibleInEditor(LevelModels, bVisible);
 
 	RequestUpdateAllLevels();
 }
 	
-void FLevelCollectionModel::ShowLevels(const FLevelModelList& InLevelList)
+void FLevelCollectionModel::ShowLevelsInEditor(const FLevelModelList& InLevelList)
 {
 	if (IsReadOnly())
 	{
@@ -530,19 +546,55 @@ void FLevelCollectionModel::ShowLevels(const FLevelModelList& InLevelList)
 		LevelModels.Add(It->Get());
 		bVisible.Add(true);
 	}
-	FLevelModel::SetVisible(LevelModels, bVisible);
+	FLevelModel::SetVisibleInEditor(LevelModels, bVisible);
 
 	RequestUpdateAllLevels();
 }
 
-void FLevelCollectionModel::ShowOnlySelectedLevels()
+void FLevelCollectionModel::ShowInEditorOnlySelectedLevels()
+{
+	ShowInEditorOnlySelectedLevels_Executed();
+}
+
+void FLevelCollectionModel::ShowInEditorAllButSelectedLevels()
+{
+	ShowInEditorAllButSelectedLevels_Executed();
+}
+
+void FLevelCollectionModel::HideLevelsInGame(const FLevelModelList& InLevelList)
+{
+	// For efficiency, set visibility of all levels at once
+	TArray<FLevelModel*> LevelModels;
+	TArray<bool> bVisible;
+	for (auto It = InLevelList.CreateConstIterator(); It; ++It)
+	{
+		LevelModels.Add(It->Get());
+		bVisible.Add(false);
+	}
+	FLevelModel::SetVisibleInGame(LevelModels, bVisible);
+}
+
+void FLevelCollectionModel::ShowLevelsInGame(const FLevelModelList& InLevelList)
+{
+	// For efficiency, set visibility of all levels at once
+	TArray<FLevelModel*> LevelModels;
+	TArray<bool> bVisible;
+	for (auto It = InLevelList.CreateConstIterator(); It; ++It)
+	{
+		LevelModels.Add(It->Get());
+		bVisible.Add(true);
+	}
+	FLevelModel::SetVisibleInGame(LevelModels, bVisible);
+}
+
+void FLevelCollectionModel::ShowInGameOnlySelectedLevels()
 {
-	ShowOnlySelectedLevels_Executed();
+	ShowInGameOnlySelectedLevels_Executed();
 }
 
-void FLevelCollectionModel::ShowAllButSelectedLevels()
+void FLevelCollectionModel::ShowInGameAllButSelectedLevels()
 {
-	ShowAllButSelectedLevels_Executed();
+	ShowInGameAllButSelectedLevels_Executed();
 }
 
 void FLevelCollectionModel::UnlockLevels(const FLevelModelList& InLevelList)
@@ -603,7 +655,7 @@ void FLevelCollectionModel::SaveLevels(const FLevelModelList& InLevelList)
 	{
 		if ((*It)->GetLevelObject())
 		{
-			if (!(*It)->IsVisible())
+			if (!(*It)->IsVisibleInEditor())
 			{
 				FMessageDialog::Open( EAppMsgType::Ok, NSLOCTEXT("UnrealEd", "UnableToSaveInvisibleLevels", "Save aborted.  Levels must be made visible before they can be saved.") );
 				return;
@@ -827,7 +879,7 @@ bool FLevelCollectionModel::PassesAllFilters(const FLevelModel& Item) const
 	return false;
 }
 
-void FLevelCollectionModel::BuildHierarchyMenu(FMenuBuilder& InMenuBuilder) const
+void FLevelCollectionModel::BuildHierarchyMenu(FMenuBuilder& InMenuBuilder, EBuildHierarchyMenuFlags Flags) const
 {
 }
 
@@ -1011,7 +1063,7 @@ bool FLevelCollectionModel::AreAllSelectedLevelsEditableAndVisible() const
 	for (auto It = SelectedLevelsList.CreateConstIterator(); It; ++It)
 	{
 		if ((*It)->IsEditable() == false ||
-			(*It)->IsVisible() == false)
+			(*It)->IsVisibleInEditor() == false)
 		{
 			return false;
 		}
@@ -1038,7 +1090,7 @@ bool FLevelCollectionModel::AreAnySelectedLevelsEditableAndVisible() const
 	for (auto It = SelectedLevelsList.CreateConstIterator(); It; ++It)
 	{
 		if ((*It)->IsEditable() == true && 
-			(*It)->IsVisible() == true)
+			(*It)->IsVisibleInEditor() == true)
 		{
 			return true;
 		}
@@ -1047,6 +1099,16 @@ bool FLevelCollectionModel::AreAnySelectedLevelsEditableAndVisible() const
 	return false;
 }
 
+bool FLevelCollectionModel::CanExecuteGameVisibilityCommandsForSelectedLevels() const
+{
+	return AreAnySelectedLevelsLoaded() && CanExecuteGameVisibilityCommands();
+}
+
+bool FLevelCollectionModel::CanExecuteGameVisibilityCommands() const
+{
+	return !WorldHierarchy::IsInPie();
+}
+
 bool FLevelCollectionModel::IsSelectedLevelEditable() const
 {
 	if (SelectedLevelsList.Num() == 1)
@@ -1485,7 +1547,7 @@ FBox FLevelCollectionModel::GetVisibleLevelsBoundingBox(const FLevelModelList& I
 			TotalBounds+= GetVisibleLevelsBoundingBox((*It)->GetChildren(), bIncludeChildren);
 		}
 		
-		if ((*It)->IsVisible())
+		if ((*It)->IsVisibleInEditor())
 		{
 			TotalBounds+= (*It)->GetLevelBounds();
 		}
@@ -1624,47 +1686,89 @@ void FLevelCollectionModel::InvertSelection_Executed()
 	SetSelectedLevels(InvertedLevels);
 }
 
-void FLevelCollectionModel::ShowSelectedLevels_Executed()
+void FLevelCollectionModel::ShowInEditorSelectedLevels_Executed()
 {
-	ShowLevels(GetSelectedLevels());
+	ShowLevelsInEditor(GetSelectedLevels());
 }
 
-void FLevelCollectionModel::HideSelectedLevels_Executed()
+void FLevelCollectionModel::HideInEditorSelectedLevels_Executed()
 {
-	HideLevels(GetSelectedLevels());
+	HideLevelsInEditor(GetSelectedLevels());
 }
 
-void FLevelCollectionModel::ShowOnlySelectedLevels_Executed()
+void FLevelCollectionModel::ShowInEditorOnlySelectedLevels_Executed()
 {
 	//stash off a copy of the original array, as setting visibility can destroy the selection
 	FLevelModelList SelectedLevelsCopy = GetSelectedLevels();
 	
 	InvertSelection_Executed();
-	HideSelectedLevels_Executed();
+	HideInEditorSelectedLevels_Executed();
 	SetSelectedLevels(SelectedLevelsCopy);
-	ShowSelectedLevels_Executed();
+	ShowInEditorSelectedLevels_Executed();
 }
 
-void FLevelCollectionModel::ShowAllButSelectedLevels_Executed()
+void FLevelCollectionModel::ShowInEditorAllButSelectedLevels_Executed()
 {
 	//stash off a copy of the original array, as setting visibility can destroy the selection
 	FLevelModelList SelectedLevelsCopy = GetSelectedLevels();
 
 	InvertSelection_Executed();
-	ShowSelectedLevels_Executed();
+	ShowInEditorSelectedLevels_Executed();
 	SetSelectedLevels(SelectedLevelsCopy);
-	HideSelectedLevels_Executed();
+	HideInEditorSelectedLevels_Executed();
+}
+
+
+void FLevelCollectionModel::ShowInEditorAllLevels_Executed()
+{
+	ShowLevelsInEditor(GetFilteredLevels());
+}
+
+void FLevelCollectionModel::HideInEditorAllLevels_Executed()
+{
+	HideLevelsInEditor(GetFilteredLevels());
 }
 
+void FLevelCollectionModel::ShowInGameSelectedLevels_Executed()
+{
+	ShowLevelsInGame(GetSelectedLevels());
+}
 
-void FLevelCollectionModel::ShowAllLevels_Executed()
+void FLevelCollectionModel::HideInGameSelectedLevels_Executed()
 {
-	ShowLevels(GetFilteredLevels());
+	HideLevelsInGame(GetSelectedLevels());
 }
 
-void FLevelCollectionModel::HideAllLevels_Executed()
+void FLevelCollectionModel::ShowInGameOnlySelectedLevels_Executed()
 {
-	HideLevels(GetFilteredLevels());
+	//stash off a copy of the original array, as setting visibility can destroy the selection
+	FLevelModelList SelectedLevelsCopy = GetSelectedLevels();
+	
+	InvertSelection_Executed();
+	HideInGameSelectedLevels_Executed();
+	SetSelectedLevels(SelectedLevelsCopy);
+	ShowInGameSelectedLevels_Executed();
+}
+
+void FLevelCollectionModel::ShowInGameAllButSelectedLevels_Executed()
+{
+	//stash off a copy of the original array, as setting visibility can destroy the selection
+	FLevelModelList SelectedLevelsCopy = GetSelectedLevels();
+
+	InvertSelection_Executed();
+	ShowInGameSelectedLevels_Executed();
+	SetSelectedLevels(SelectedLevelsCopy);
+	HideInGameSelectedLevels_Executed();
+}
+
+void FLevelCollectionModel::ShowInGameAllLevels_Executed()
+{
+	ShowLevelsInGame(GetFilteredLevels());
+}
+
+void FLevelCollectionModel::HideInGameAllLevels_Executed()
+{
+	HideLevelsInGame(GetFilteredLevels());
 }
 
 void FLevelCollectionModel::LockSelectedLevels_Executed()
@@ -1869,16 +1973,28 @@ void FLevelCollectionModel::FillLockSubMenu(FMenuBuilder& InMenuBuilder)
 	}
 }
 
-void FLevelCollectionModel::FillVisibilitySubMenu(FMenuBuilder& InMenuBuilder)
+void FLevelCollectionModel::FillEditorVisibilitySubMenu(FMenuBuilder& InMenuBuilder)
+{
+	const FLevelCollectionCommands& Commands = FLevelCollectionCommands::Get();
+
+	InMenuBuilder.AddMenuEntry( Commands.World_ShowInEditorSelectedLevels );
+	InMenuBuilder.AddMenuEntry( Commands.World_HideInEditorSelectedLevels );
+	InMenuBuilder.AddMenuEntry( Commands.World_ShowInEditorOnlySelectedLevels );
+	InMenuBuilder.AddMenuEntry( Commands.World_ShowInEditorAllButSelectedLevels );
+	InMenuBuilder.AddMenuEntry( Commands.World_ShowInEditorAllLevels );
+	InMenuBuilder.AddMenuEntry( Commands.World_HideInEditorAllLevels );
+}
+
+void FLevelCollectionModel::FillGameVisibilitySubMenu(FMenuBuilder& InMenuBuilder)
 {
 	const FLevelCollectionCommands& Commands = FLevelCollectionCommands::Get();
 
-	InMenuBuilder.AddMenuEntry( Commands.World_ShowSelectedLevels );
-	InMenuBuilder.AddMenuEntry( Commands.World_HideSelectedLevels );
-	InMenuBuilder.AddMenuEntry( Commands.World_ShowOnlySelectedLevels );
-	InMenuBuilder.AddMenuEntry( Commands.World_ShowAllButSelectedLevels );
-	InMenuBuilder.AddMenuEntry( Commands.World_ShowAllLevels );
-	InMenuBuilder.AddMenuEntry( Commands.World_HideAllLevels );
+	InMenuBuilder.AddMenuEntry( Commands.World_ShowInGameSelectedLevels );
+	InMenuBuilder.AddMenuEntry( Commands.World_HideInGameSelectedLevels );
+	InMenuBuilder.AddMenuEntry( Commands.World_ShowInGameOnlySelectedLevels );
+	InMenuBuilder.AddMenuEntry( Commands.World_ShowInGameAllButSelectedLevels );
+	InMenuBuilder.AddMenuEntry( Commands.World_ShowInGameAllLevels );
+	InMenuBuilder.AddMenuEntry( Commands.World_HideInGameAllLevels );
 }
 
 void FLevelCollectionModel::FillSourceControlSubMenu(FMenuBuilder& InMenuBuilder)
diff --git a/Engine/Source/Editor/WorldBrowser/Private/LevelCollectionModel.h b/Engine/Source/Editor/WorldBrowser/Private/LevelCollectionModel.h
--- a/Engine/Source/Editor/WorldBrowser/Private/LevelCollectionModel.h
+++ b/Engine/Source/Editor/WorldBrowser/Private/LevelCollectionModel.h
@@ -7,6 +7,7 @@
 #include "Engine/World.h"
 #include "TickableEditorObject.h"
 #include "WorldBrowserDragDrop.h"
+#include "Misc/EnumClassFlags.h"
 #include "Misc/IFilter.h"
 #include "LevelModel.h"
 
@@ -20,21 +21,25 @@ class UMaterialInterface;
 typedef IFilter< const FLevelModel* >				LevelFilter;
 typedef TFilterCollection< const FLevelModel* >		LevelFilterCollection;
 
-/**
- * Interface for non-UI presentation logic for a world
- */
-class WORLDBROWSER_API FLevelCollectionModel
+enum class EBuildHierarchyMenuFlags
+{
+	None,
+	ShowGameVisibility = 1 << 0
+};
+ENUM_CLASS_FLAGS(EBuildHierarchyMenuFlags);
+
+/** Interface for non-UI presentation logic for a world. */
+class FLevelCollectionModel
 	: public TSharedFromThis<FLevelCollectionModel>	
 	, public FTickableEditorObject
 {
 public:
+	
 	DECLARE_EVENT_OneParam( FLevelCollectionModel, FOnNewItemAdded, TSharedPtr<FLevelModel>);
 	DECLARE_EVENT( FLevelCollectionModel, FSimpleEvent );
-
-
-public:
+	
 	FLevelCollectionModel();
-	virtual ~FLevelCollectionModel();
+	virtual ~FLevelCollectionModel() override;
 
 	/** FTickableEditorObject interface */
 	virtual void Tick( float DeltaTime ) override;
@@ -96,21 +101,32 @@ public:
 	/**	@return	Found level model with specified level package name */
 	TSharedPtr<FLevelModel> FindLevelModel(const FName& PackageName) const;
 
-	/**	Hides level in the world */
-	void HideLevels(const FLevelModelList& InLevelList);
-	
-	/**	Shows level in the world */
-	void ShowLevels(const FLevelModelList& InLevelList);
+	/**	Hides level in editor worlds */
+	void HideLevelsInEditor(const FLevelModelList& InLevelList);
+
+	/**	Shows level in editor worlds */
+	void ShowLevelsInEditor(const FLevelModelList& InLevelList);
 
 	/** Toggles the selected levels to a visible state; toggles all other levels to an invisible state */
-	void ShowOnlySelectedLevels();
+	void ShowInEditorOnlySelectedLevels();
 
 	/** Toggles the selected levels to an invisible state; toggles all other levels to a visible state */
-	void ShowAllButSelectedLevels();
+	void ShowInEditorAllButSelectedLevels();
+	
+	/**	Hides level in game worlds */
+	void HideLevelsInGame(const FLevelModelList& InLevelList);
+
+	/**	Shows level in game worlds */
+	void ShowLevelsInGame(const FLevelModelList& InLevelList);
+
+	/** Toggles the selected levels to a visible state for game worlds; toggles all other levels to an invisible state */
+	void ShowInGameOnlySelectedLevels();
+
+	/** Toggles the selected levels to an invisible state for game worlds; toggles all other levels to a visible state */
+	void ShowInGameAllButSelectedLevels();
 
 	/**	Unlocks level in the world */
 	void UnlockLevels(const FLevelModelList& InLevelList);
-	
 	/**	Locks level in the world */
 	void LockLevels(const FLevelModelList& InLevelList);
 
@@ -125,7 +141,6 @@ public:
 
 	/**	Loads level from disk */
 	void LoadLevels(const FLevelModelList& InLevelList);
-	
 	/**	Unloads levels from the editor */
 	virtual void UnloadLevels(const FLevelModelList& InLevelList);
 
@@ -154,7 +169,7 @@ public:
 	virtual bool PassesAllFilters(const FLevelModel& InLevelModel) const;
 	
 	/**	Builds 'hierarchy' commands menu for a selected levels */
-	virtual void BuildHierarchyMenu(FMenuBuilder& InMenuBuilder) const;
+	virtual void BuildHierarchyMenu(FMenuBuilder& InMenuBuilder, EBuildHierarchyMenuFlags Flags) const;
 	
 	/**	Customize 'File' section in main menu  */
 	virtual void CustomizeFileMainMenu(FMenuBuilder& InMenuBuilder) const;
@@ -236,6 +251,12 @@ public:
 
 	/** @return whether any of the currently selected levels is editable and visible*/
 	bool AreAnySelectedLevelsEditableAndVisible() const;
+
+	/** @return Whether game visibility of the selected levels is allowed to be changed. */
+	bool CanExecuteGameVisibilityCommandsForSelectedLevels() const;
+	
+	/** @return Whether game visibility of levels is allowed to be changed. */
+	bool CanExecuteGameVisibilityCommands() const;
 	
 	/** @return whether currently only one level selected and it is editable */
 	bool IsSelectedLevelEditable() const;
@@ -411,7 +432,7 @@ protected:
 
 	/** Inverts level selection in the collection view model */
 	void InvertSelection_Executed();
-	
+
 	/** Adds the Actors in the selected Levels from the viewport's existing selection */
 	void SelectActors_Executed();
 
@@ -421,23 +442,41 @@ protected:
 	/** Set level `Use External Actors` to bExternal  */
 	void ConvertLevelToExternalActors_Executed(bool bExternal);
 
-	/** Toggles selected levels to a visible state in the viewports */
-	void ShowSelectedLevels_Executed();
+	/** Toggles selected levels to a visible state in the viewports for editor worlds */
+	void ShowInEditorSelectedLevels_Executed();
+
+	/** Toggles selected levels to an invisible state in the viewports for editor worlds */
+	void HideInEditorSelectedLevels_Executed();
+
+	/** Toggles the selected levels to a visible state for editor worlds; toggles all other levels to an invisible state. */
+	void ShowInEditorOnlySelectedLevels_Executed();
+
+	/** Toggles the selected levels to an invisible state for editor worlds; toggles all other levels to a visible state. */
+	void ShowInEditorAllButSelectedLevels_Executed();
+
+	/** Toggles all levels to a visible state in the viewports for editor worlds */
+	void ShowInEditorAllLevels_Executed();
+
+	/** Hides all levels to an invisible state in the viewports for editor worlds */
+	void HideInEditorAllLevels_Executed();
+
+	/** Toggles selected levels to a visible state in the viewports for game worlds */
+	void ShowInGameSelectedLevels_Executed();
 
 	/** Toggles selected levels to an invisible state in the viewports */
-	void HideSelectedLevels_Executed();
+	void HideInGameSelectedLevels_Executed();
 
-	/** Toggles the selected levels to a visible state; toggles all other levels to an invisible state */
-	void ShowOnlySelectedLevels_Executed();
-	
-	/** Toggles the selected levels to an invisible state; toggles all other levels to a visible state */
-	void ShowAllButSelectedLevels_Executed();
+	/** Toggles the selected levels to a visible state for game worlds; toggles all other levels to an invisible state. */
+	void ShowInGameOnlySelectedLevels_Executed();
+
+	/** Toggles the selected levels to an invisible state for game worlds; toggles all other levels to a visible state. */
+	void ShowInGameAllButSelectedLevels_Executed();
 
-	/** Toggles all levels to a visible state in the viewports */
-	void ShowAllLevels_Executed();
+	/** Toggles all levels to a visible state in the viewports for game worlds */
+	void ShowInGameAllLevels_Executed();
 
-	/** Hides all levels to an invisible state in the viewports */
-	void HideAllLevels_Executed();
+	/** Hides all levels to an invisible state in the viewports for game worlds */
+	void HideInGameAllLevels_Executed();
 	
 	/** Locks selected levels */
 	void LockSelectedLevels_Executed();
@@ -456,7 +495,7 @@ protected:
 
 	/** Unlocks all levels */
 	void UnlockAllLevels_Executed();
-
+	
 	/** Toggle all read-only levels */
 	void ToggleReadOnlyLevels_Executed();
 
@@ -485,13 +524,16 @@ protected:
 	}
 	
 	/** Fills MenuBulder with Lock level related commands */
-	void FillLockSubMenu(class FMenuBuilder& MenuBuilder);
-	
+	void FillLockSubMenu(FMenuBuilder& MenuBuilder);
+
+	/** Fills MenuBulder with level visisbility related commands */
+	void FillEditorVisibilitySubMenu(FMenuBuilder& MenuBuilder);
+
 	/** Fills MenuBulder with level visisbility related commands */
-	void FillVisibilitySubMenu(class FMenuBuilder& MenuBuilder);
+	void FillGameVisibilitySubMenu(FMenuBuilder& MenuBuilder);
 
 	/** Fills MenuBulder with SCC related commands */
-	void FillSourceControlSubMenu(class FMenuBuilder& MenuBuilder);
+	void FillSourceControlSubMenu(FMenuBuilder& MenuBuilder);
 				
 protected:
 	/**  */
diff --git a/Engine/Source/Editor/WorldBrowser/Private/LevelModel.cpp b/Engine/Source/Editor/WorldBrowser/Private/LevelModel.cpp
--- a/Engine/Source/Editor/WorldBrowser/Private/LevelModel.cpp
+++ b/Engine/Source/Editor/WorldBrowser/Private/LevelModel.cpp
@@ -191,7 +191,7 @@ bool FLevelModel::IsLoading() const
 	return (LevelCollectionModel.IsSimulating() ? SimulationStatus.bLoading : bLoadingLevel);
 }
 
-bool FLevelModel::IsVisible() const
+bool FLevelModel::IsVisibleInEditor() const
 {
 	if (LevelCollectionModel.IsSimulating())
 	{
@@ -215,6 +215,24 @@ bool FLevelModel::IsVisible() const
 	}
 }
 
+bool FLevelModel::IsVisibleInGame() const
+{
+	ULevel* Level = GetLevelObject();
+	if (!Level)
+	{
+		return false;
+	}
+
+	// Unreal does not allow the persistent level to be invisible in game.
+	if (Level->IsPersistentLevel())
+	{
+		return true;
+	}
+	
+	ULevelStreaming* StreamingLevel = FLevelUtils::FindStreamingLevel(Level);
+	return StreamingLevel && StreamingLevel->GetShouldBeVisibleFlag();
+}
+
 bool FLevelModel::IsLocked() const
 {
 	ULevel* Level = GetLevelObject();
@@ -261,14 +279,14 @@ void FLevelModel::LoadLevel()
 
 }
 
-void FLevelModel::SetVisible(bool bVisible)
+void FLevelModel::SetVisibleInEditor(bool bVisible)
 {
 	TArray<FLevelModel*> LevelModels({ this });
 	TArray<bool> bAreVisible({ bVisible });
-	FLevelModel::SetVisible(LevelModels, bAreVisible);
+	FLevelModel::SetVisibleInEditor(LevelModels, bAreVisible);
 }
 
-void FLevelModel::SetVisible(TArray<FLevelModel*>& LevelModels, const TArray<bool>& bAreVisible)
+void FLevelModel::SetVisibleInEditor(TArray<FLevelModel*>& LevelModels, const TArray<bool>& bAreVisible)
 {
 	// Don't create unnecessary transactions
 	if (LevelModels.Num() == 0 || LevelModels.Num() != bAreVisible.Num())
@@ -283,7 +301,7 @@ void FLevelModel::SetVisible(TArray<FLevelModel*>& LevelModels, const TArray<boo
 	{
 		FLevelModel* LevelModel = LevelModels[LevelModelIdx];
 		// If visibility does not change, omit it
-		if (LevelModel->IsVisible() == bAreVisible[LevelModelIdx])
+		if (LevelModel->IsVisibleInEditor() == bAreVisible[LevelModelIdx])
 		{
 			continue;
 		}
@@ -318,6 +336,31 @@ void FLevelModel::SetVisible(TArray<FLevelModel*>& LevelModels, const TArray<boo
 	}
 }
 
+void FLevelModel::SetVisibleInGame(TArray<FLevelModel*>& LevelModels, const TArray<bool>& bAreVisible)
+{
+	FScopedTransaction Transaction(LOCTEXT("ToggleGameVisibility", "Toggle Game Visibility"));
+	bool bMadeChange = false;
+	
+	for (int32 i = 0; i < LevelModels.Num(); ++i)
+	{
+		ULevel* Level = LevelModels[i]->GetLevelObject();
+		ULevelStreaming* LevelStreaming = FLevelUtils::FindStreamingLevel(Level);
+		const bool bShouldBeVisible = bAreVisible[i];
+		if (LevelStreaming && LevelStreaming->GetShouldBeVisibleFlag() != bShouldBeVisible)
+		{
+			bMadeChange = true;
+			Level->Modify();
+			LevelStreaming->Modify();
+			LevelStreaming->SetShouldBeVisible(bShouldBeVisible);
+		}
+	}
+
+	if (!bMadeChange)
+	{
+		Transaction.Cancel();
+	}
+}
+
 void FLevelModel::SetLocked(bool bLocked)
 {
 	if (LevelCollectionModel.IsReadOnly())
@@ -426,7 +469,7 @@ void FLevelModel::MakeLevelCurrent()
 		}
 							
 		// Force the current level to be visible.
-		LevelCollectionModel.ShowLevels({ AsShared() });
+		LevelCollectionModel.ShowLevelsInEditor({ AsShared() });
 	}
 	else
 	{
diff --git a/Engine/Source/Editor/WorldBrowser/Private/LevelModel.h b/Engine/Source/Editor/WorldBrowser/Private/LevelModel.h
--- a/Engine/Source/Editor/WorldBrowser/Private/LevelModel.h
+++ b/Engine/Source/Editor/WorldBrowser/Private/LevelModel.h
@@ -117,8 +117,11 @@ public:
 	/** @return Whether level is in process of loading content */
 	bool IsLoading() const;
 
-	/**	@return Whether level is visible in the world */
-	bool IsVisible() const;
+	/**	@return Whether level is visible in editor worlds */
+	bool IsVisibleInEditor() const;
+
+	/**	@return Whether level is visible in game worlds */
+	bool IsVisibleInGame() const;
 
 	/**	@return Whether level is locked */
 	bool IsLocked() const;
@@ -136,10 +139,12 @@ public:
 	virtual void LoadLevel();
 
 	/** Sets the Level's visibility */
-	virtual void SetVisible(bool bVisible);
+	virtual void SetVisibleInEditor(bool bVisible);
 
 	/** Static function analog to the non-static SetVisible but applied to an TArray of FLevelModel elements */
-	static void SetVisible(TArray<FLevelModel*>& LevelModels, const TArray<bool>& bAreVisible);
+	static void SetVisibleInEditor(TArray<FLevelModel*>& LevelModels, const TArray<bool>& bAreVisible);
+
+	static void SetVisibleInGame(TArray<FLevelModel*>& LevelModels, const TArray<bool>& bAreVisible);
 
 	/** Sets the Level's locked/unlocked state */
 	void SetLocked(bool bLocked);
diff --git a/Engine/Source/Editor/WorldBrowser/Private/SWorldHierarchyImpl.cpp b/Engine/Source/Editor/WorldBrowser/Private/SWorldHierarchyImpl.cpp
--- a/Engine/Source/Editor/WorldBrowser/Private/SWorldHierarchyImpl.cpp
+++ b/Engine/Source/Editor/WorldBrowser/Private/SWorldHierarchyImpl.cpp
@@ -25,6 +25,7 @@
 #include "LevelFolders.h"
 #include "ScopedTransaction.h"
 #include "Editor.h"
+#include "WorldBrowserStyle.h"
 #include "WorldHierarchyColumns.h"
 #include "RevisionControlStyle/RevisionControlStyle.h"
 
@@ -287,7 +288,8 @@ TSharedRef<SHeaderRow> SWorldHierarchyImpl::CreateHeaderRow()
 {
 	using namespace UE::WorldHierarchy;
 	constexpr float IconWidth = 24.f;
-	return SAssignNew(HeaderRowWidget, SHeaderRow)
+	
+	TSharedRef<SHeaderRow> HeaderRow = SAssignNew(HeaderRowWidget, SHeaderRow)
 		.CanSelectGeneratedColumn(true) // Lets the user choose which columns should be shown
 
 		/** Level visibility column */
@@ -303,6 +305,19 @@ TSharedRef<SHeaderRow> SWorldHierarchyImpl::CreateHeaderRow()
 			.ToolTipText(LOCTEXT("Column.Visibility.Tooltip", "Toggles the editor visibility"))
 		]
 
+		/** Level game visibility column */
+		+ SHeaderRow::Column(HierarchyColumns::ColumnID_GameVisibility)
+		.Visibility(bFoldersOnlyMode ? EVisibility::Collapsed : EVisibility::Visible)
+		.DefaultLabel(LOCTEXT("Column.GameVisibility.DefaultLabel", "Game Visibility"))
+		.FixedWidth(IconWidth)
+		.HeaderContent()
+		[
+			SNew(SImage)
+			.Image(WorldBrowser::FWorldBrowserStyle::Get().GetBrush( "WorldBrowser.VisibleInGame" ))
+			.ColorAndOpacity(FSlateColor::UseForeground())
+			.ToolTipText(LOCTEXT("Column.GameVisibility.Tooltip", "Toggles the visibility in games"))
+		]
+
 		/** LevelName label column */
 		+ SHeaderRow::Column( HierarchyColumns::ColumnID_LevelLabel )
 		.FillWidth(1.f)
@@ -392,6 +407,10 @@ TSharedRef<SHeaderRow> SWorldHierarchyImpl::CreateHeaderRow()
 			.Image(FAppStyle::Get().GetBrush("Level.ColorIcon"))
 			.ToolTipText(LOCTEXT("Column.Color.Tooltip", "Color used for visualization of Level"))
 		];
+
+	// By default, hide this column; it is only used in certain workflows, like Virtual Production. The user can manually show this column.
+	SetColumnVisible(HierarchyColumns::ColumnID_GameVisibility, false);
+	return HeaderRow;
 }
 
 TSharedRef<ITableRow> SWorldHierarchyImpl::GenerateTreeRow(WorldHierarchy::FWorldTreeItemPtr Item, const TSharedRef<STableViewBase>& OwnerTable)
@@ -501,7 +520,11 @@ TSharedPtr<SWidget> SWorldHierarchyImpl::ConstructLevelContextMenu() const
 
 		Menu->AddDynamicSection("HierarchyDynamicSection", FNewToolMenuDelegateLegacy::CreateLambda([this](FMenuBuilder& MenuBuilder, UToolMenu* Menu)
 		{
-			WorldModel->BuildHierarchyMenu(MenuBuilder);
+			const bool bIsGameVisibilityColumnVisible = IsColumnVisible(UE::WorldHierarchy::HierarchyColumns::ColumnID_GameVisibility);
+			WorldModel->BuildHierarchyMenu(
+				MenuBuilder,
+				bIsGameVisibilityColumnVisible ? EBuildHierarchyMenuFlags::ShowGameVisibility : EBuildHierarchyMenuFlags::None
+				);
 		}));
 
 		// Generate the "Move To" and "Select" submenus based on the current selection
diff --git a/Engine/Source/Editor/WorldBrowser/Private/SWorldHierarchyItem.cpp b/Engine/Source/Editor/WorldBrowser/Private/SWorldHierarchyItem.cpp
--- a/Engine/Source/Editor/WorldBrowser/Private/SWorldHierarchyItem.cpp
+++ b/Engine/Source/Editor/WorldBrowser/Private/SWorldHierarchyItem.cpp
@@ -21,6 +21,7 @@
 
 #include "WorldTreeItemTypes.h"
 #include "LevelFolders.h"
+#include "WorldBrowserStyle.h"
 #include "Framework/Application/SlateApplication.h"
 #include "Styling/StyleColors.h"
 
@@ -168,12 +169,12 @@ TSharedRef< SWidget > SWorldHierarchyItem::GenerateWidgetForColumn( const FName&
 		else if (ColumnID == HierarchyColumns::ColumnID_EditorVisibility)
 		{
 			TableRowContent =
-				SAssignNew(VisibilityButton, SButton)
+				SAssignNew(EditorVisibilityButton, SButton)
 				.ContentPadding(0)
 				.ButtonStyle(FAppStyle::Get(), "SimpleButton")
-				.IsEnabled(this, &SWorldHierarchyItem::IsVisibilityEnabled)
-				.OnClicked(this, &SWorldHierarchyItem::OnToggleVisibility)
-				.ToolTipText(this, &SWorldHierarchyItem::GetVisibilityToolTip)
+				.IsEnabled(this, &SWorldHierarchyItem::IsEditorVisibilityEnabled)
+				.OnClicked(this, &SWorldHierarchyItem::OnToggleEditorVisibility)
+				.ToolTipText(this, &SWorldHierarchyItem::GetEditorVisibilityToolTip)
 				.HAlign(HAlign_Center)
 				.VAlign(VAlign_Center)
 				.ForegroundColor_Lambda([this]() { return IsSelected() ? FSlateColor::UseForeground() : FSlateColor::UseStyle(); })
@@ -185,6 +186,26 @@ TSharedRef< SWidget > SWorldHierarchyItem::GenerateWidgetForColumn( const FName&
 				]
 			;
 		}
+		else if (ColumnID == HierarchyColumns::ColumnID_GameVisibility)
+		{
+			TableRowContent =
+			   SAssignNew(GameVisibilityButton, SButton)
+			   .ContentPadding(0)
+			   .ButtonStyle(FAppStyle::Get(), "SimpleButton")
+			   .IsEnabled(this, &SWorldHierarchyItem::IsGameVisibilityEnabled)
+			   .OnClicked(this, &SWorldHierarchyItem::OnToggleGameVisibility)
+			   .ToolTipText(this, &SWorldHierarchyItem::GetGameVisibilityToolTip)
+			   .HAlign(HAlign_Center)
+			   .VAlign(VAlign_Center)
+			   .ForegroundColor_Lambda([this]() { return IsSelected() ? FSlateColor::UseForeground() : FSlateColor::UseStyle(); })
+			   .Content()
+			   [
+				  SNew(SImage)
+				  .Image(this, &SWorldHierarchyItem::GetGameVisibilityBrush)
+				  .ColorAndOpacity(FSlateColor::UseForeground())
+			   ]
+			;
+		}
 		else if (ColumnID == HierarchyColumns::ColumnID_Color)
 		{
 			TableRowContent =
@@ -378,9 +399,14 @@ bool SWorldHierarchyItem::IsLockEnabled() const
 	return WorldTreeItem->HasLockControls();
 }
 
-bool SWorldHierarchyItem::IsVisibilityEnabled() const
+bool SWorldHierarchyItem::IsEditorVisibilityEnabled() const
+{
+	return WorldTreeItem->HasEditorVisibilityControls();
+}
+
+bool SWorldHierarchyItem::IsGameVisibilityEnabled() const
 {
-	return WorldTreeItem->HasVisibilityControls();
+	return WorldTreeItem->HasGameVisibilityControls();
 }
 
 bool SWorldHierarchyItem::IsKismetEnabled() const
@@ -393,22 +419,42 @@ FSlateColor SWorldHierarchyItem::GetDrawColor() const
 	return WorldTreeItem->GetDrawColor();
 }
 
-FReply SWorldHierarchyItem::OnToggleVisibility()
+FReply SWorldHierarchyItem::OnToggleEditorVisibility()
+{
+	if (FSlateApplication::Get().GetModifierKeys().AreModifersDown(EModifierKey::Alt))
+	{
+		if (WorldTreeItem->IsVisibleInEditor()) 
+		{
+			WorldTreeItem->OnShowInEditorAllButSelected();
+		}
+		else
+		{
+			WorldTreeItem->OnShowInEditorOnlySelected();
+		}
+	}
+	else
+	{
+		WorldTreeItem->OnToggleEditorVisibility();
+	}
+	return FReply::Handled();
+}
+
+FReply SWorldHierarchyItem::OnToggleGameVisibility()
 {
 	if (FSlateApplication::Get().GetModifierKeys().AreModifersDown(EModifierKey::Alt))
 	{
-		if (WorldTreeItem->IsVisible()) 
+		if (WorldTreeItem->IsVisibleInEditor()) 
 		{
-			WorldTreeItem->OnShowAllButSelected();
+			WorldTreeItem->OnShowInGameAllButSelected();
 		}
 		else
 		{
-			WorldTreeItem->OnShowOnlySelected();
+			WorldTreeItem->OnShowInGameOnlySelected();
 		}
 	}
 	else
 	{
-		WorldTreeItem->OnToggleVisibility();
+		WorldTreeItem->OnToggleGameVisibility();
 	}
 	return FReply::Handled();
 }
@@ -513,9 +559,14 @@ EVisibility SWorldHierarchyItem::GetColorButtonVisibility() const
 	return Result;
 }
 
-FText SWorldHierarchyItem::GetVisibilityToolTip() const
+FText SWorldHierarchyItem::GetEditorVisibilityToolTip() const
+{
+	return WorldTreeItem->GetEditorVisibilityToolTipText();
+}
+
+FText SWorldHierarchyItem::GetGameVisibilityToolTip() const
 {
-	return WorldTreeItem->GetVisibilityToolTipText();
+	return WorldTreeItem->GetGameVisibilityToolTipText();
 }
 
 FText SWorldHierarchyItem::GetSaveToolTip() const
@@ -675,17 +726,19 @@ const FSlateBrush* SWorldHierarchyItem::GetLevelIconBrush() const
 
 const FSlateBrush* SWorldHierarchyItem::GetLevelVisibilityBrush() const
 {
-	if (WorldTreeItem->HasVisibilityControls())
+	if (WorldTreeItem->HasEditorVisibilityControls())
 	{
-		if (WorldTreeItem->IsVisible())
+		if (WorldTreeItem->IsVisibleInEditor())
 		{
-			return VisibilityButton->IsHovered() ? FAppStyle::Get().GetBrush( "Level.VisibleHighlightIcon16x" ) :
-													FAppStyle::Get().GetBrush( "Level.VisibleIcon16x" );
+			return EditorVisibilityButton->IsHovered()
+				? FAppStyle::Get().GetBrush( "Level.VisibleHighlightIcon16x" )
+				: FAppStyle::Get().GetBrush( "Level.VisibleIcon16x" );
 		}
 		else
 		{
-			return VisibilityButton->IsHovered() ? FAppStyle::Get().GetBrush( "Level.NotVisibleHighlightIcon16x" ) :
-													FAppStyle::Get().GetBrush( "Level.NotVisibleIcon16x" );
+			return EditorVisibilityButton->IsHovered()
+				? FAppStyle::Get().GetBrush( "Level.NotVisibleHighlightIcon16x" )
+				: FAppStyle::Get().GetBrush( "Level.NotVisibleIcon16x" );
 		}
 	}
 	else
@@ -694,6 +747,13 @@ const FSlateBrush* SWorldHierarchyItem::GetLevelVisibilityBrush() const
 	}
 }
 
+const FSlateBrush* SWorldHierarchyItem::GetGameVisibilityBrush() const
+{
+	return WorldTreeItem->IsVisibleInGame()
+	   ? WorldBrowser::FWorldBrowserStyle::Get().GetBrush( "WorldBrowser.VisibleInGame" )
+	   : WorldBrowser::FWorldBrowserStyle::Get().GetBrush( "WorldBrowser.HiddenInGame" );
+}
+
 const FSlateBrush* SWorldHierarchyItem::GetLightingScenarioBrush() const
 {
 	if (WorldTreeItem->IsLightingScenario())
diff --git a/Engine/Source/Editor/WorldBrowser/Private/SWorldHierarchyItem.h b/Engine/Source/Editor/WorldBrowser/Private/SWorldHierarchyItem.h
--- a/Engine/Source/Editor/WorldBrowser/Private/SWorldHierarchyItem.h
+++ b/Engine/Source/Editor/WorldBrowser/Private/SWorldHierarchyItem.h
@@ -70,7 +70,8 @@ private:
 	bool IsSaveEnabled() const;
 	bool IsLightingScenarioEnabled() const;
 	bool IsLockEnabled() const;
-	bool IsVisibilityEnabled() const;
+	bool IsEditorVisibilityEnabled() const;
+	bool IsGameVisibilityEnabled() const;
 	bool IsKismetEnabled() const;
 
 	/** Get DrawColor for the level */
@@ -81,7 +82,8 @@ private:
 	 *
 	 *	@return	A reply that indicated whether this event was handled.
 	 */
-	FReply OnToggleVisibility();
+	FReply OnToggleEditorVisibility();
+	FReply OnToggleGameVisibility();
 
 	FReply OnToggleLightingScenario();
 
@@ -125,8 +127,11 @@ private:
 	/** Whether color button should be visible, depends on whether sub-level is loaded */
 	EVisibility GetColorButtonVisibility() const;
 	
-	/** Gets the tooltip for the visibility toggle */
-	FText GetVisibilityToolTip() const;
+	/** Gets the tooltip for the editor visibility toggle */
+	FText GetEditorVisibilityToolTip() const;
+	
+	/** Gets the tooltip for the game visibility toggle */
+	FText GetGameVisibilityToolTip() const;
 
 	/** Gets the tooltip for the save button */
 	FText GetSaveToolTip() const;
@@ -162,6 +167,7 @@ private:
 	 *	@return	The SlateBrush representing the Level's visibility state
 	 */
 	const FSlateBrush* GetLevelVisibilityBrush() const;
+	const FSlateBrush* GetGameVisibilityBrush() const;
 
 	const FSlateBrush* GetLightingScenarioBrush() const;
 
@@ -229,8 +235,11 @@ private:
 	/** True when this item has children and is expanded */
 	TAttribute<bool>				IsItemExpanded;
 
-	/**	The visibility button for the Level */
-	TSharedPtr<SButton>				VisibilityButton;
+	/**	The editor visibility button for the Level */
+	TSharedPtr<SButton>				EditorVisibilityButton;
+	
+	/**	The game visibility button for the Level */
+	TSharedPtr<SButton>				GameVisibilityButton;
 
 	/**	The lighting scenario button for the Level */
 	TSharedPtr<SButton>				LightingScenarioButton;
diff --git a/Engine/Source/Editor/WorldBrowser/Private/StreamingLevels/StreamingLevelCollectionModel.cpp b/Engine/Source/Editor/WorldBrowser/Private/StreamingLevels/StreamingLevelCollectionModel.cpp
--- a/Engine/Source/Editor/WorldBrowser/Private/StreamingLevels/StreamingLevelCollectionModel.cpp
+++ b/Engine/Source/Editor/WorldBrowser/Private/StreamingLevels/StreamingLevelCollectionModel.cpp
@@ -257,7 +257,7 @@ TSharedPtr<WorldHierarchy::FWorldBrowserDragDropOp> FStreamingLevelCollectionMod
 	}
 }
 
-void FStreamingLevelCollectionModel::BuildHierarchyMenu(FMenuBuilder& InMenuBuilder) const
+void FStreamingLevelCollectionModel::BuildHierarchyMenu(FMenuBuilder& InMenuBuilder, EBuildHierarchyMenuFlags Flags) const
 {
 	const FLevelCollectionCommands& Commands = FLevelCollectionCommands::Get();
 
@@ -281,11 +281,21 @@ void FStreamingLevelCollectionModel::BuildHierarchyMenu(FMenuBuilder& InMenuBuil
 			InMenuBuilder.AddMenuEntry( Commands.World_MakeLevelCurrent );
 		}
 		
-		// Visibility commands
+		// Editor Visibility commands
 		InMenuBuilder.AddSubMenu( 
-			LOCTEXT("VisibilityHeader", "Visibility"),
-			LOCTEXT("VisibilitySubMenu_ToolTip", "Selected Level(s) visibility commands"),
-			FNewMenuDelegate::CreateSP(const_cast<FStreamingLevelCollectionModel*>(this), &FStreamingLevelCollectionModel::FillVisibilitySubMenu ) );
+			LOCTEXT("EditorVisibilityHeader", "Visibility in Editor"),
+			LOCTEXT("EditorVisibilitySubMenu_ToolTip", "Selected Level(s) visibility commands for editor worlds"),
+			FNewMenuDelegate::CreateSP(const_cast<FStreamingLevelCollectionModel*>(this), &FStreamingLevelCollectionModel::FillEditorVisibilitySubMenu ) );
+		
+		// Game Visibility commands
+		if (EnumHasAnyFlags(Flags, EBuildHierarchyMenuFlags::ShowGameVisibility))
+		{
+			InMenuBuilder.AddSubMenu( 
+				LOCTEXT("GameVisibilityHeader", "Visibility in Game"),
+				LOCTEXT("GameVisibilitySubMenu_ToolTip", "Selected Level(s) visibility commands for game worlds"),
+				FNewMenuDelegate::CreateSP(const_cast<FStreamingLevelCollectionModel*>(this), &FStreamingLevelCollectionModel::FillGameVisibilitySubMenu ) );
+		}
+		
 
 		// Lock commands
 		InMenuBuilder.AddSubMenu( 
@@ -543,7 +553,8 @@ void FStreamingLevelCollectionModel::MergeSelectedLevels_Executed()
 	FLevelModelList SelectedLevelsCopy = SelectedLevelsList;
 
 	//make sure the selected levels are made visible (and thus fully loaded) before merging
-	ShowSelectedLevels_Executed();
+	ShowInEditorSelectedLevels_Executed();
+	ShowInGameSelectedLevels_Executed();
 
 	//restore the original selection and select all actors in the selected levels
 	SetSelectedLevels(SelectedLevelsCopy);
diff --git a/Engine/Source/Editor/WorldBrowser/Private/StreamingLevels/StreamingLevelCollectionModel.h b/Engine/Source/Editor/WorldBrowser/Private/StreamingLevels/StreamingLevelCollectionModel.h
--- a/Engine/Source/Editor/WorldBrowser/Private/StreamingLevels/StreamingLevelCollectionModel.h
+++ b/Engine/Source/Editor/WorldBrowser/Private/StreamingLevels/StreamingLevelCollectionModel.h
@@ -38,7 +38,7 @@ public:
 	virtual void AddExistingLevelsFromAssetData(const TArray<FAssetData>& WorldList) override;
 	virtual TSharedPtr<WorldHierarchy::FWorldBrowserDragDropOp> CreateDragDropOp() const override;
 	virtual TSharedPtr<WorldHierarchy::FWorldBrowserDragDropOp> CreateDragDropOp(const FLevelModelList& InLevels) const override;
-	virtual void BuildHierarchyMenu(FMenuBuilder& InMenuBuilder) const override;
+	virtual void BuildHierarchyMenu(FMenuBuilder& InMenuBuilder, EBuildHierarchyMenuFlags Flags) const override;
 	virtual void CustomizeFileMainMenu(FMenuBuilder& InMenuBuilder) const override;
 	virtual void RegisterDetailsCustomization(class FPropertyEditorModule& PropertyModule, TSharedPtr<class IDetailsView> InDetailsView)  override;
 	virtual void UnregisterDetailsCustomization(class FPropertyEditorModule& PropertyModule, TSharedPtr<class IDetailsView> InDetailsView) override;
diff --git a/Engine/Source/Editor/WorldBrowser/Private/StreamingLevels/StreamingLevelCustomization.cpp b/Engine/Source/Editor/WorldBrowser/Private/StreamingLevels/StreamingLevelCustomization.cpp
--- a/Engine/Source/Editor/WorldBrowser/Private/StreamingLevels/StreamingLevelCustomization.cpp
+++ b/Engine/Source/Editor/WorldBrowser/Private/StreamingLevels/StreamingLevelCustomization.cpp
@@ -429,7 +429,7 @@ bool FStreamingLevelCustomization::LevelViewportTransformAllowed() const
 	if (CollectionModel.IsValid() && CollectionModel->IsOneLevelSelected())
 	{
 		auto SelectedLevel = CollectionModel->GetSelectedLevels()[0];
-		return SelectedLevel->IsEditable() && SelectedLevel->IsVisible();
+		return SelectedLevel->IsEditable() && SelectedLevel->IsVisibleInEditor();
 	}
 
 	return false;
diff --git a/Engine/Source/Editor/WorldBrowser/Private/Tiles/SWorldTileItem.cpp b/Engine/Source/Editor/WorldBrowser/Private/Tiles/SWorldTileItem.cpp
--- a/Engine/Source/Editor/WorldBrowser/Private/Tiles/SWorldTileItem.cpp
+++ b/Engine/Source/Editor/WorldBrowser/Private/Tiles/SWorldTileItem.cpp
@@ -389,7 +389,7 @@ bool SWorldTileItem::IsItemEnabled() const
 {
 	if (WorldModel->IsSimulating())
 	{
-		return TileModel->IsVisible();
+		return TileModel->IsVisibleInEditor();
 	}
 	else
 	{
diff --git a/Engine/Source/Editor/WorldBrowser/Private/Tiles/WorldTileCollectionModel.cpp b/Engine/Source/Editor/WorldBrowser/Private/Tiles/WorldTileCollectionModel.cpp
--- a/Engine/Source/Editor/WorldBrowser/Private/Tiles/WorldTileCollectionModel.cpp
+++ b/Engine/Source/Editor/WorldBrowser/Private/Tiles/WorldTileCollectionModel.cpp
@@ -342,11 +342,17 @@ void FWorldTileCollectionModel::BuildWorldCompositionMenu(FMenuBuilder& InMenuBu
 			InMenuBuilder.AddMenuEntry(Commands.World_UnloadLevel);
 			InMenuBuilder.AddMenuEntry(Commands.World_SaveSelectedLevels);
 		
-			// Visibility commands
+			// Editor Visibility commands
 			InMenuBuilder.AddSubMenu( 
-				LOCTEXT("VisibilityHeader", "Visibility"),
-				LOCTEXT("VisibilitySubMenu_ToolTip", "Selected Level(s) visibility commands"),
-				FNewMenuDelegate::CreateSP(const_cast<FWorldTileCollectionModel*>(this), &FWorldTileCollectionModel::FillVisibilitySubMenu ) );
+				LOCTEXT("EditorVisibilityHeader", "Visibility in Editor"),
+				LOCTEXT("EditorVisibilitySubMenu_ToolTip", "Selected Level(s) visibility commands for editor worlds"),
+				FNewMenuDelegate::CreateSP(const_cast<FWorldTileCollectionModel*>(this), &FWorldTileCollectionModel::FillEditorVisibilitySubMenu ) );
+
+			// Game Visibility commands
+			InMenuBuilder.AddSubMenu( 
+				LOCTEXT("GameVisibilityHeader", "Visibility in Game"),
+				LOCTEXT("GameVisibilitySubMenu_ToolTip", "Selected Level(s) visibility commands for game worlds"),
+				FNewMenuDelegate::CreateSP(const_cast<FWorldTileCollectionModel*>(this), &FWorldTileCollectionModel::FillGameVisibilitySubMenu ) );
 
 			// Lock commands
 			InMenuBuilder.AddSubMenu( 
@@ -434,7 +440,7 @@ void FWorldTileCollectionModel::BuildWorldCompositionMenu(FMenuBuilder& InMenuBu
 	InMenuBuilder.EndSection();
 }
 
-void FWorldTileCollectionModel::BuildHierarchyMenu(FMenuBuilder& InMenuBuilder) const
+void FWorldTileCollectionModel::BuildHierarchyMenu(FMenuBuilder& InMenuBuilder, EBuildHierarchyMenuFlags Flags) const
 {
 	const FLevelCollectionCommands& Commands = FLevelCollectionCommands::Get();
 		
@@ -454,10 +460,19 @@ void FWorldTileCollectionModel::BuildHierarchyMenu(FMenuBuilder& InMenuBuilder)
 				
 		// Visibility commands
 		InMenuBuilder.AddSubMenu( 
-			LOCTEXT("VisibilityHeader", "Visibility"),
-			LOCTEXT("VisibilitySubMenu_ToolTip", "Selected Level(s) visibility commands"),
-			FNewMenuDelegate::CreateSP(const_cast<FWorldTileCollectionModel*>(this), &FWorldTileCollectionModel::FillVisibilitySubMenu ) );
+			LOCTEXT("EditorVisibilityHeader", "Visibility in Editor"),
+			LOCTEXT("EditorVisibilitySubMenu_ToolTip", "Selected Level(s) visibility commands for editor"),
+			FNewMenuDelegate::CreateSP(const_cast<FWorldTileCollectionModel*>(this), &FWorldTileCollectionModel::FillEditorVisibilitySubMenu ) );
 
+		// Game Visibility commands
+		if (EnumHasAnyFlags(Flags, EBuildHierarchyMenuFlags::ShowGameVisibility))
+		{
+			InMenuBuilder.AddSubMenu( 
+	            LOCTEXT("GameVisibilityHeader", "Visibility in Game"),
+	            LOCTEXT("GameVisibilitySubMenu_ToolTip", "Selected Level(s) visibility commands for game worlds"),
+	            FNewMenuDelegate::CreateSP(const_cast<FWorldTileCollectionModel*>(this), &FWorldTileCollectionModel::FillGameVisibilitySubMenu ) );
+		}
+		
 		// Lock commands
 		InMenuBuilder.AddSubMenu( 
 			LOCTEXT("LockHeader", "Lock"),
@@ -1425,7 +1440,7 @@ void FWorldTileCollectionModel::ClearParentLink_Executed()
 bool FWorldTileCollectionModel::CanAddLandscapeProxy(FWorldTileModel::EWorldDirections InWhere) const
 {
 	if (SelectedLevelsList.Num() == 1 &&
-		SelectedLevelsList[0]->IsVisible() &&
+		SelectedLevelsList[0]->IsVisibleInEditor() &&
 		StaticCastSharedPtr<FWorldTileModel>(SelectedLevelsList[0])->IsLandscapeBased())
 	{
 		return true;
@@ -1474,9 +1489,9 @@ void FWorldTileCollectionModel::AddLandscapeProxy_Executed(FWorldTileModel::EWor
 
 		ALandscapeProxy* SourceLandscape = LandscapeTileModel->GetLandscape();
 	
-		NewLevelModel->SetVisible(false);
+		NewLevelModel->SetVisibleInEditor(false);
 		NewLevelModel->CreateAdjacentLandscapeProxy(SourceLandscape, InWhere);
-		ShowLevels(Levels);
+		ShowLevelsInEditor(Levels);
 	}
 }
 
@@ -1683,7 +1698,7 @@ void FWorldTileCollectionModel::ImportTiledLandscape_Executed()
 					PopulateLevelsList();
 					TSharedPtr<FWorldTileModel> NewTileModel = StaticCastSharedPtr<FWorldTileModel>(FindLevelModel(NewWorld->GetOutermost()->GetFName()));
 					// Hide level, so we do not depend on a current world origin
-					NewTileModel->SetVisible(false);
+					NewTileModel->SetVisibleInEditor(false);
 
 					// Create landscape proxy in a new level
 					ALandscapeProxy* NewLandscape = NewTileModel->ImportLandscapeTile(TileImportSettings);
@@ -1749,10 +1764,10 @@ void FWorldTileCollectionModel::ReimportTiledLandscape_Executed(FName TargetLaye
 	// Hide all visible levels
 	for (auto LevelModel : AllLevelsList)
 	{
-		AllLevelsVisibilityState.Add(LevelModel->IsVisible());
+		AllLevelsVisibilityState.Add(LevelModel->IsVisibleInEditor());
 		if (!LevelModel->IsPersistent())
 		{
-			LevelModel->SetVisible(false);
+			LevelModel->SetVisibleInEditor(false);
 		}
 	}
 
@@ -1764,7 +1779,7 @@ void FWorldTileCollectionModel::ReimportTiledLandscape_Executed(FName TargetLaye
 	{
 		Progress.EnterProgressFrame();
 
-		TileModel->SetVisible(true);
+		TileModel->SetVisibleInEditor(true);
 
 		ALandscapeProxy* Landscape = TileModel->GetLandscape();
 		FIntRect LandscapeSize = Landscape->GetBoundingRect();
@@ -1807,7 +1822,7 @@ void FWorldTileCollectionModel::ReimportTiledLandscape_Executed(FName TargetLaye
 			}
 		}
 
-		TileModel->SetVisible(false);
+		TileModel->SetVisibleInEditor(false);
 	}
 
 	// Restore world origin tracking
@@ -1818,7 +1833,7 @@ void FWorldTileCollectionModel::ReimportTiledLandscape_Executed(FName TargetLaye
 	{
 		if (AllLevelsVisibilityState[LevelIdx])
 		{
-			AllLevelsList[LevelIdx]->SetVisible(true);
+			AllLevelsList[LevelIdx]->SetVisibleInEditor(true);
 		}
 	}
 }
@@ -1986,11 +2001,11 @@ bool FWorldTileCollectionModel::GenerateLODLevels(FLevelModelList InLevelList, i
 		}
 
 		// We have to make original level visible, to correctly export it
-		const bool bVisibleLevel = TileModel->IsVisible();
+		const bool bVisibleLevel = TileModel->IsVisibleInEditor();
 		if (!bVisibleLevel)
 		{
 			GetWorld()->WorldComposition->bTemporarilyDisableOriginTracking = true;
-			TileModel->SetVisible(true);
+			TileModel->SetVisibleInEditor(true);
 		}
 
 		FLevelSimplificationDetails SimplificationDetails = TileModel->GetLevelObject()->LevelSimplification[TargetLODIndex];
@@ -2200,7 +2215,7 @@ bool FWorldTileCollectionModel::GenerateLODLevels(FLevelModelList InLevelList, i
 		// Restore level original visibility
 		if (!bVisibleLevel)
 		{
-			TileModel->SetVisible(false);
+			TileModel->SetVisibleInEditor(false);
 			GetWorld()->WorldComposition->bTemporarilyDisableOriginTracking = false;
 		}
 	
diff --git a/Engine/Source/Editor/WorldBrowser/Private/Tiles/WorldTileCollectionModel.h b/Engine/Source/Editor/WorldBrowser/Private/Tiles/WorldTileCollectionModel.h
--- a/Engine/Source/Editor/WorldBrowser/Private/Tiles/WorldTileCollectionModel.h
+++ b/Engine/Source/Editor/WorldBrowser/Private/Tiles/WorldTileCollectionModel.h
@@ -39,7 +39,7 @@ public:
 	virtual TSharedPtr<WorldHierarchy::FWorldBrowserDragDropOp> CreateDragDropOp() const override;
 	virtual TSharedPtr<WorldHierarchy::FWorldBrowserDragDropOp> CreateDragDropOp(const FLevelModelList& InLevels) const override;
 	virtual bool PassesAllFilters(const FLevelModel& InLevelModel) const override;
-	virtual void BuildHierarchyMenu(FMenuBuilder& InMenuBuilder) const override;
+	virtual void BuildHierarchyMenu(FMenuBuilder& InMenuBuilder, EBuildHierarchyMenuFlags Flags) const override;
 	virtual void CustomizeFileMainMenu(FMenuBuilder& InMenuBuilder) const override;
 	virtual bool GetPlayerView(FVector& Location, FRotator& Rotation) const override;
 	virtual bool GetObserverView(FVector& Location, FRotator& Rotation) const override;
diff --git a/Engine/Source/Editor/WorldBrowser/Private/Tiles/WorldTileModel.cpp b/Engine/Source/Editor/WorldBrowser/Private/Tiles/WorldTileModel.cpp
--- a/Engine/Source/Editor/WorldBrowser/Private/Tiles/WorldTileModel.cpp
+++ b/Engine/Source/Editor/WorldBrowser/Private/Tiles/WorldTileModel.cpp
@@ -1,33 +1,31 @@
 // Copyright Epic Games, Inc. All Rights Reserved.
 
 #include "Tiles/WorldTileModel.h"
-#include "Engine/World.h"
-#include "HAL/FileManager.h"
-#include "UObject/Package.h"
+
 #include "DragAndDrop/LevelDragDropOp.h"
-#include "Engine/LevelBounds.h"
-#include "Engine/LevelStreamingDynamic.h"
 #include "Editor.h"
 #include "Editor/Transactor.h"
-#include "ScopedTransaction.h"
 #include "EditorLevelUtils.h"
-#include "LevelCollectionModel.h"
-
-#include "Modules/ModuleManager.h"
-#include "Tiles/WorldTileDetails.h"
-#include "Tiles/WorldTileCollectionModel.h"
+#include "Engine/LevelBounds.h"
+#include "Engine/LevelStreamingDynamic.h"
+#include "Engine/World.h"
 #include "Engine/WorldComposition.h"
 #include "GameFramework/WorldSettings.h"
-#include "LandscapeInfo.h"
+#include "HAL/FileManager.h"
+#include "Landscape.h"
 #include "LandscapeEditorModule.h"
 #include "LandscapeFileFormatInterface.h"
+#include "LandscapeInfo.h"
 #include "LandscapeStreamingProxy.h"
-#include "Landscape.h"
+#include "LevelCollectionModel.h"
 #include "Modules/ModuleManager.h"
-
+#include "Modules/ModuleManager.h"
+#include "ScopedTransaction.h"
+#include "Tiles/WorldTileCollectionModel.h"
+#include "Tiles/WorldTileDetails.h"
+#include "UObject/Package.h"
 
 #define LOCTEXT_NAMESPACE "WorldBrowser"
-DEFINE_LOG_CATEGORY_STATIC(WorldBrowser, Log, All);
 
 FWorldTileModel::FWorldTileModel(FWorldTileCollectionModel& InWorldModel, int32 InTileIdx)
 	: FLevelModel(InWorldModel) 
@@ -260,7 +258,7 @@ bool FWorldTileModel::ShouldBeVisible(FBox EditableArea) const
 	return false;
 }
 
-void FWorldTileModel::SetVisible(bool bVisible)
+void FWorldTileModel::SetVisibleInEditor(bool bVisible)
 {
 	if (LevelCollectionModel.IsReadOnly())
 	{
@@ -275,7 +273,7 @@ void FWorldTileModel::SetVisible(bool bVisible)
 	}
 	
 	// Don't create unnecessary transactions
-	if (IsVisible() == bVisible)
+	if (IsVisibleInEditor() == bVisible)
 	{
 		return;
 	}
@@ -317,7 +315,7 @@ void FWorldTileModel::Shelve()
 	}
 	
 	//
-	SetVisible(false);
+	SetVisibleInEditor(false);
 	bWasShelved = true;
 }
 
@@ -329,7 +327,7 @@ void FWorldTileModel::Unshelve()
 	}
 
 	//
-	SetVisible(true);
+	SetVisibleInEditor(true);
 	bWasShelved = false;
 }
 
diff --git a/Engine/Source/Editor/WorldBrowser/Private/Tiles/WorldTileModel.h b/Engine/Source/Editor/WorldBrowser/Private/Tiles/WorldTileModel.h
--- a/Engine/Source/Editor/WorldBrowser/Private/Tiles/WorldTileModel.h
+++ b/Engine/Source/Editor/WorldBrowser/Private/Tiles/WorldTileModel.h
@@ -83,7 +83,7 @@ public:
 	virtual void Update() override;
 	virtual void UpdateAsset(const FAssetData& AssetData) override;
 	virtual void LoadLevel() override;
-	virtual void SetVisible(bool bVisible) override;
+	virtual void SetVisibleInEditor(bool bVisible) override;
 	virtual FVector2D GetLevelPosition2D() const override;
 	virtual FVector2D GetLevelSize2D() const override;
 	virtual void OnDrop(const TSharedPtr<FLevelDragDropOp>& Op) override;
diff --git a/Engine/Source/Editor/WorldBrowser/Private/Tiles/WorldTileThumbnails.cpp b/Engine/Source/Editor/WorldBrowser/Private/Tiles/WorldTileThumbnails.cpp
--- a/Engine/Source/Editor/WorldBrowser/Private/Tiles/WorldTileThumbnails.cpp
+++ b/Engine/Source/Editor/WorldBrowser/Private/Tiles/WorldTileThumbnails.cpp
@@ -41,7 +41,7 @@ FSlateTextureDataPtr FTileThumbnail::UpdateThumbnail()
 	}
 	
 	// Load image from a package header
-	if (!TileModel.IsVisible() || TileModel.IsSimulating())
+	if (!TileModel.IsVisibleInEditor() || TileModel.IsSimulating())
 	{
 		const FName LevelAssetName = TileModel.GetAssetName();
 		TSet<FName> ObjectFullNames;
diff --git a/Engine/Source/Editor/WorldBrowser/Private/WorldBrowserStyle.cpp b/Engine/Source/Editor/WorldBrowser/Private/WorldBrowserStyle.cpp
new file mode 100644
--- /dev/null
+++ b/Engine/Source/Editor/WorldBrowser/Private/WorldBrowserStyle.cpp
@@ -0,0 +1,21 @@
+﻿// Copyright Epic Games, Inc. All Rights Reserved.
+
+#include "WorldBrowserStyle.h"
+
+#include "Styling/SlateStyleMacros.h"
+
+namespace WorldBrowser
+{
+	FWorldBrowserStyle::FWorldBrowserStyle()
+		: FSlateStyleSet("WorldBrowserStyle")
+	{
+		const FVector2D Icon16x16(16.0f, 16.0f);
+		SetContentRoot(FPaths::EngineContentDir() / TEXT("Slate"));
+		SetCoreContentRoot(FPaths::EngineContentDir() / TEXT("Slate"));
+		
+		Set("WorldBrowser.VisibleInGame", new CORE_IMAGE_BRUSH_SVG("Starship/Common/VisibleInGame", Icon16x16));
+		Set("WorldBrowser.HiddenInGame", new CORE_IMAGE_BRUSH_SVG("Starship/Common/HiddenInGame", Icon16x16));
+		
+		FSlateStyleRegistry::RegisterSlateStyle(*this);
+	}
+}
diff --git a/Engine/Source/Editor/WorldBrowser/Private/WorldBrowserStyle.h b/Engine/Source/Editor/WorldBrowser/Private/WorldBrowserStyle.h
new file mode 100644
--- /dev/null
+++ b/Engine/Source/Editor/WorldBrowser/Private/WorldBrowserStyle.h
@@ -0,0 +1,24 @@
+﻿// Copyright Epic Games, Inc. All Rights Reserved.
+
+#pragma once
+
+#include "Styling/SlateStyleRegistry.h"
+#include "Styling/SlateStyle.h"
+
+namespace WorldBrowser
+{
+	class FWorldBrowserStyle final : public FSlateStyleSet
+	{
+	public:
+
+		FWorldBrowserStyle();
+		virtual ~FWorldBrowserStyle() { FSlateStyleRegistry::UnRegisterSlateStyle(*this); }
+	
+		static FWorldBrowserStyle& Get()
+		{
+			static FWorldBrowserStyle Inst;
+			return Inst;
+		}
+	};
+}
+
diff --git a/Engine/Source/Editor/WorldBrowser/Private/WorldHierarchyColumns.cpp b/Engine/Source/Editor/WorldBrowser/Private/WorldHierarchyColumns.cpp
--- a/Engine/Source/Editor/WorldBrowser/Private/WorldHierarchyColumns.cpp
+++ b/Engine/Source/Editor/WorldBrowser/Private/WorldHierarchyColumns.cpp
@@ -6,10 +6,11 @@ namespace UE::WorldHierarchy::HierarchyColumns
 {
 	const FName ColumnID_LevelLabel( "Level" );
 	const FName ColumnID_EditorVisibility( "EditorVisibility" );
+	const FName ColumnID_GameVisibility( "GameVisibility" );
 	const FName ColumnID_LightingScenario( "LightingScenario" );
 	const FName ColumnID_Lock( "Lock" );
 	const FName ColumnID_SCCStatus( "SCC_Status" );
 	const FName ColumnID_Save( "Save" );
 	const FName ColumnID_Color("Color");
 	const FName ColumnID_Kismet( "Blueprint" );
-}
\ No newline at end of file
+}
diff --git a/Engine/Source/Editor/WorldBrowser/Private/WorldTreeItemTypes.cpp b/Engine/Source/Editor/WorldBrowser/Private/WorldTreeItemTypes.cpp
--- a/Engine/Source/Editor/WorldBrowser/Private/WorldTreeItemTypes.cpp
+++ b/Engine/Source/Editor/WorldBrowser/Private/WorldTreeItemTypes.cpp
@@ -1,6 +1,8 @@
 // Copyright Epic Games, Inc. All Rights Reserved.
 
 #include "WorldTreeItemTypes.h"
+
+#include "Editor.h"
 #include "Templates/Casts.h"
 #include "LevelModel.h"
 #include "LevelCollectionModel.h"
@@ -47,11 +49,17 @@
 
 		return OutWorlds.Num() > 0;
 	}
-
-	//------------------------
-	// FLevelModelTreeItem
-	//------------------------
 	
+	bool IsInPie()
+	{
+		return GEditor && GEditor->GetPlayInEditorSessionInfo().IsSet();
+	}
+
+	//------------------------
+	// FLevelModelTreeItem
+	//------------------------
+
+
 	FLevelModelTreeItem::FLevelModelTreeItem(TSharedRef<FLevelModel> InLevelModel)
 		: LevelModel(InLevelModel)
 		, bPersistentLevel(InLevelModel->IsPersistent())
@@ -130,9 +138,24 @@
 		return LockToolTip;
 	}
 
-	FText FLevelModelTreeItem::GetVisibilityToolTipText() const
+	FText FLevelModelTreeItem::GetEditorVisibilityToolTipText() const
 	{
-		return LOCTEXT("VisibilityButtonToolTip", "Toggle Level Visibility");
+		return LOCTEXT("EditorVisibilityButtonToolTip", "Toggle whether Level is visible in the editor");
 	}
 
+	FText FLevelModelTreeItem::GetGameVisibilityToolTipText() const
+	{
+		if (IsPersistentLevel())
+		{
+			return LOCTEXT("LevelModel.LevelGameVisibilityButtonToolTip.PersistentLevel", "Game visibility cannot be toggled for persistent level.");
+		}
+
+		if (IsInPie())
+		{
+			return LOCTEXT("LevelModel.LevelGameVisibilityButtonToolTip.Pie", "Game visibility cannot be toggled during PIE.");
+		}
+		
+		return LOCTEXT("LevelModel.GameVisibilityButtonToolTip.Normal", "Toggle whether Level is visible in-game");
+	}
+
 	FText FLevelModelTreeItem::GetSaveToolTipText() const
@@ -228,11 +251,19 @@
 		return IsLoaded() ;
 	}
 
-	bool FLevelModelTreeItem::HasVisibilityControls() const
+	bool FLevelModelTreeItem::HasEditorVisibilityControls() const
 	{
 		return IsLoaded();
 	}
 
+	bool FLevelModelTreeItem::HasGameVisibilityControls() const
+	{
+		return HasEditorVisibilityControls() 
+			// Unreal generally does not allow hiding the persistent level (root level) - so no controls if this item is the persistent level. 
+			&& !IsPersistentLevel()
+			&& !IsInPie();
+	}
+
 	bool FLevelModelTreeItem::HasColorButtonControls() const
 	{
 		// The root level does not have a color button
@@ -285,43 +316,54 @@
 		}
 	}
 
-	bool FLevelModelTreeItem::IsVisible() const
+	bool FLevelModelTreeItem::IsVisibleInEditor() const
 	{
-		return LevelModel.IsValid() ? LevelModel.Pin()->IsVisible() : false;
+		return LevelModel.IsValid() ? LevelModel.Pin()->IsVisibleInEditor() : false;
 	}
 
-	void FLevelModelTreeItem::OnToggleVisibility()
+	bool FLevelModelTreeItem::IsVisibleInGame() const
+	{
+		return LevelModel.IsValid() ? LevelModel.Pin()->IsVisibleInGame() : false;
+	}
+
+	void FLevelModelTreeItem::OnToggleEditorVisibility()
 	{
 		if (LevelModel.IsValid())
 		{
-			SetVisible(!LevelModel.Pin()->IsVisible());
+			SetVisibleInEditor(!LevelModel.Pin()->IsVisibleInEditor());
 		}
 	}
 
-	void FLevelModelTreeItem::OnShowOnlySelected()
+	void FLevelModelTreeItem::OnShowInEditorOnlySelected()
 	{
-		FLevelModelList LevelsToChange;
-		PopulateLevelModelList(LevelsToChange);
-		if (GetLevelSelectionFlag())
-		{
-			FLevelModelList CurrentLevels = WorldModel.Pin()->GetSelectedLevels();
-			LevelsToChange.Append(CurrentLevels);
-		}
-		WorldModel.Pin()->SetSelectedLevels(LevelsToChange);
-		WorldModel.Pin()->ShowOnlySelectedLevels();
+		SetSelectedLevelsToPopulatedList();
+		WorldModel.Pin()->ShowInEditorOnlySelectedLevels();
 	}
 
-	void FLevelModelTreeItem::OnShowAllButSelected()
+	void FLevelModelTreeItem::OnShowInEditorAllButSelected()
 	{
-		FLevelModelList LevelsToChange;
-		PopulateLevelModelList(LevelsToChange);
-		if (GetLevelSelectionFlag())
-		{
-			FLevelModelList CurrentLevels = WorldModel.Pin()->GetSelectedLevels();
-			LevelsToChange.Append(CurrentLevels);
-		}
-		WorldModel.Pin()->SetSelectedLevels(LevelsToChange);
-		WorldModel.Pin()->ShowAllButSelectedLevels();
+		SetSelectedLevelsToPopulatedList();
+		WorldModel.Pin()->ShowInEditorAllButSelectedLevels();
 	}
 
+	void FLevelModelTreeItem::OnToggleGameVisibility()
+	{
+		if (LevelModel.IsValid())
+		{
+			SetVisibleInGame(!LevelModel.Pin()->IsVisibleInGame());
+		}
+	}
+
+	void FLevelModelTreeItem::OnShowInGameOnlySelected()
+	{
+		SetSelectedLevelsToPopulatedList();
+		WorldModel.Pin()->ShowInGameOnlySelectedLevels();
+	}
+
+	void FLevelModelTreeItem::OnShowInGameAllButSelected()
+	{
+		SetSelectedLevelsToPopulatedList();
+		WorldModel.Pin()->ShowInGameAllButSelectedLevels();
+	}
+
 	void FLevelModelTreeItem::PopulateLevelModelList(FLevelModelList& InModelList)
@@ -332,7 +374,7 @@
 		}
 	}
 
-	void FLevelModelTreeItem::SetVisible(bool bVisible)
+	void FLevelModelTreeItem::SetVisibleInEditor(bool bVisible)
 	{
 		FLevelModelList LevelModels;
 
@@ -343,11 +385,30 @@
 
 		if (bVisible)
 		{
-			WorldModel.Pin()->ShowLevels(LevelModels);
+			WorldModel.Pin()->ShowLevelsInEditor(LevelModels);
 		}
 		else
 		{
-			WorldModel.Pin()->HideLevels(LevelModels);
+			WorldModel.Pin()->HideLevelsInEditor(LevelModels);
 		}
 	}
 
+	void FLevelModelTreeItem::SetVisibleInGame(bool bVisible)
+	{
+		FLevelModelList LevelModels;
+
+		if (LevelModel.IsValid())
+		{
+			LevelModels.Add(LevelModel.Pin());
+		}
+
+		if (bVisible)
+		{
+			WorldModel.Pin()->ShowLevelsInGame(LevelModels);
+		}
+		else
+		{
+			WorldModel.Pin()->HideLevelsInGame(LevelModels);
+		}
+	}
+
@@ -370,27 +431,13 @@
 
 	void FLevelModelTreeItem::OnLockOnlySelected()
 	{
-		FLevelModelList LevelsToChange;
-		PopulateLevelModelList(LevelsToChange);
-		if (GetLevelSelectionFlag())
-		{
-			FLevelModelList CurrentLevels = WorldModel.Pin()->GetSelectedLevels();
-			LevelsToChange.Append(CurrentLevels);
-		}
-		WorldModel.Pin()->SetSelectedLevels(LevelsToChange);
+		SetSelectedLevelsToPopulatedList();
 		WorldModel.Pin()->LockOnlySelectedLevels();
 	}
 
 	void FLevelModelTreeItem::OnLockAllButSelected()
 	{
-		FLevelModelList LevelsToChange;
-		PopulateLevelModelList(LevelsToChange);
-		if (GetLevelSelectionFlag())
-		{
-			FLevelModelList CurrentLevels = WorldModel.Pin()->GetSelectedLevels();
-			LevelsToChange.Append(CurrentLevels);
-		}
-		WorldModel.Pin()->SetSelectedLevels(LevelsToChange);
+		SetSelectedLevelsToPopulatedList();
 		WorldModel.Pin()->LockAllButSelectedLevels();
 	}
 
@@ -577,6 +624,24 @@
 			}
 		}
 	}
+	
+	void FLevelModelTreeItem::SetSelectedLevelsToPopulatedList()
+	{
+		FLevelModelList LevelsToChange;
+		PopulateLevelModelList(LevelsToChange);
+		if (GetLevelSelectionFlag())
+		{
+			FLevelModelList CurrentLevels = WorldModel.Pin()->GetSelectedLevels();
+			LevelsToChange.Append(CurrentLevels);
+		}
+		WorldModel.Pin()->SetSelectedLevels(LevelsToChange);
+	}
+
+	bool FLevelModelTreeItem::IsPersistentLevel() const
+	{
+		const TSharedPtr<FLevelModel> LevelModelPin = LevelModel.Pin();
+		return LevelModelPin && LevelModelPin->IsPersistent();
+	}
 
 
 	//------------------------
@@ -690,9 +755,19 @@
 		return LockToolTip;
 	}
 
-	FText FFolderTreeItem::GetVisibilityToolTipText() const
+	FText FFolderTreeItem::GetEditorVisibilityToolTipText() const
 	{
-		return LOCTEXT("FolderVisibilityButtonToolTip", "Toggle Visibility for All Levels");
+		return LOCTEXT("FolderEditorVisibilityButtonToolTip", "Toggle Editor Visibility for All Levels");
 	}
 
+	FText FFolderTreeItem::GetGameVisibilityToolTipText() const
+	{
+		if (IsInPie())
+		{
+			return LOCTEXT("FolderModel.GameVisibilityButtonToolTip.Pie", "Game visibility cannot be toggled during PIE.");
+		}
+		
+		return LOCTEXT("FolderModel.FolderGameVisibilityButtonToolTip", "Toggle Game Visibility for All Levels");
+	}
+
 	FText FFolderTreeItem::GetSaveToolTipText() const
@@ -722,12 +797,18 @@
 		return GetLevelModels().Num() == 0 || IsAnyChildLoaded();
 	}
 
-	bool FFolderTreeItem::HasVisibilityControls() const
+	bool FFolderTreeItem::HasEditorVisibilityControls() const
 	{
 		// If the folder has no level models associated with it, always show the visibility icon
 		return GetLevelModels().Num() == 0 || IsAnyChildLoaded();
 	}
 
+	bool FFolderTreeItem::HasGameVisibilityControls() const
+	{
+		return HasEditorVisibilityControls()
+			&& !IsInPie();
+	}
+
 	bool FFolderTreeItem::HasValidPackage() const
 	{
 		// Return true only if all children have valid packages
@@ -765,25 +846,7 @@
 		});
 	}
 
-	bool FFolderTreeItem::IsVisible() const
+	FLevelModelList FFolderTreeItem::SetSelectionToFolderChildren()
 	{
-		for (FWorldTreeItemPtr Child : Children)
-		{
-			if (Child->IsVisible())
-			{
-				return true;
-			}
-		}
-
-		return false;
-	}
-
-	void FFolderTreeItem::OnToggleVisibility()
-	{
-		SetVisible(!IsVisible());
-	}
-
-	void FFolderTreeItem::OnShowOnlySelected()
-	{
 		// This can be triggered on the non selected folder, so get its children instead
 		FLevelModelList LevelsToChange;
@@ -797,24 +860,66 @@
 			LevelsToChange.Append(CurrentLevels);
 		}
 		WorldModel.Pin()->SetSelectedLevels(LevelsToChange);
-		WorldModel.Pin()->ShowOnlySelectedLevels();
+		return LevelsToChange;
 	}
 
-	void FFolderTreeItem::OnShowAllButSelected()
+	bool FFolderTreeItem::IsVisibleInEditor() const
 	{
-		// This can be triggered on the non selected folder, so get its children instead
-		FLevelModelList LevelsToChange;
-		for (auto& Child : Children)
+		for (FWorldTreeItemPtr Child : Children)
 		{
-			Child->PopulateLevelModelList(LevelsToChange);
+			if (Child->IsVisibleInEditor())
+			{
+				return true;
+			}
 		}
-		if (GetLevelSelectionFlag())
-		{
-			FLevelModelList CurrentLevels = WorldModel.Pin()->GetSelectedLevels();
-			LevelsToChange.Append(CurrentLevels);
-		}
-		WorldModel.Pin()->SetSelectedLevels(LevelsToChange);
-		WorldModel.Pin()->ShowAllButSelectedLevels();
+
+		return false;
 	}
 
+	bool FFolderTreeItem::IsVisibleInGame() const
+	{
+		for (FWorldTreeItemPtr Child : Children)
+		{
+			if (Child->IsVisibleInGame())
+			{
+				return true;
+			}
+		}
+		return false;
+	}
+
+	void FFolderTreeItem::OnToggleEditorVisibility()
+	{
+		SetVisibleInEditor(!IsVisibleInEditor());
+	}
+
+	void FFolderTreeItem::OnShowInEditorOnlySelected()
+	{
+		SetSelectionToFolderChildren();
+		WorldModel.Pin()->ShowInEditorOnlySelectedLevels();
+	}
+
+	void FFolderTreeItem::OnShowInEditorAllButSelected()
+	{
+		SetSelectionToFolderChildren();
+		WorldModel.Pin()->ShowInEditorAllButSelectedLevels();
+	}
+
+	void FFolderTreeItem::OnToggleGameVisibility()
+	{
+		SetVisibleInGame(!IsVisibleInGame());
+	}
+
+	void FFolderTreeItem::OnShowInGameOnlySelected()
+	{
+		SetSelectionToFolderChildren();
+		WorldModel.Pin()->ShowInGameOnlySelectedLevels();
+	}
+
+	void FFolderTreeItem::OnShowInGameAllButSelected()
+	{
+		SetSelectionToFolderChildren();
+		WorldModel.Pin()->ShowInGameAllButSelectedLevels();
+	}
+
 	void FFolderTreeItem::PopulateLevelModelList(FLevelModelList& InModelList)
@@ -825,11 +930,19 @@
 		}
 	}
 
-	void FFolderTreeItem::SetVisible(bool bVisible)
+	void FFolderTreeItem::SetVisibleInEditor(bool bVisible)
 	{
 		for (auto& Child : Children)
 		{
-			Child->SetVisible(bVisible);
+			Child->SetVisibleInEditor(bVisible);
 		}
 	}
 
+	void FFolderTreeItem::SetVisibleInGame(bool bVisible)
+	{
+		for (auto& Child : Children)
+		{
+			Child->SetVisibleInGame(bVisible);
+		}
+	}
+
@@ -853,35 +966,13 @@
 
 	void FFolderTreeItem::OnLockOnlySelected()
 	{
-		// This can be triggered on the non selected folder, so get its children instead
-		FLevelModelList LevelsToChange;
-		for (auto& Child : Children)
-		{
-			Child->PopulateLevelModelList(LevelsToChange);
-		}
-		if (GetLevelSelectionFlag())
-		{
-			FLevelModelList CurrentLevels = WorldModel.Pin()->GetSelectedLevels();
-			LevelsToChange.Append(CurrentLevels);
-		}
-		WorldModel.Pin()->SetSelectedLevels(LevelsToChange);
+		SetSelectionToFolderChildren();
 		WorldModel.Pin()->LockOnlySelectedLevels();
 	}
 
 	void FFolderTreeItem::OnLockAllButSelected()
 	{
-		// This can be triggered on the non selected folder, so get its children instead
-		FLevelModelList LevelsToChange;
-		for (auto& Child : Children)
-		{
-			Child->PopulateLevelModelList(LevelsToChange);
-		}
-		if (GetLevelSelectionFlag())
-		{
-			FLevelModelList CurrentLevels = WorldModel.Pin()->GetSelectedLevels();
-			LevelsToChange.Append(CurrentLevels);
-		}
-		WorldModel.Pin()->SetSelectedLevels(LevelsToChange);
+		SetSelectionToFolderChildren();
 		WorldModel.Pin()->LockAllButSelectedLevels();
 	}
 
diff --git a/Engine/Source/Editor/WorldBrowser/Private/WorldTreeItemTypes.h b/Engine/Source/Editor/WorldBrowser/Private/WorldTreeItemTypes.h
--- a/Engine/Source/Editor/WorldBrowser/Private/WorldTreeItemTypes.h
+++ b/Engine/Source/Editor/WorldBrowser/Private/WorldTreeItemTypes.h
@@ -9,6 +9,9 @@ namespace WorldHierarchy
 	static const int32 FolderSortPriority = 10;
 	static const int32 LevelModelSortPriority = 0;
 
+	/** @return Whether currently playing in editor. */
+	bool IsInPie();
+
 	/** The tree item for the level models */
 	struct FLevelModelTreeItem : IWorldTreeItem
 	{
@@ -25,13 +28,10 @@ namespace WorldHierarchy
 		virtual FString GetDisplayString() const override;
 
 		virtual FText GetToolTipText() const override;
-
 		virtual FText GetLockToolTipText() const override;
-
-		virtual FText GetVisibilityToolTipText() const override;
-
+		virtual FText GetEditorVisibilityToolTipText() const override;
+		virtual FText GetGameVisibilityToolTipText() const override;
 		virtual FText GetSaveToolTipText() const override;
-
 		virtual FString GetPackageFileName() const override;
 
 		virtual FWorldTreeItemID GetParentID() const override;
@@ -48,7 +48,8 @@ namespace WorldHierarchy
 
 		virtual int32 GetSortPriority() const override { return LevelModelSortPriority; }
 
-		virtual bool IsVisible() const override;
+		virtual bool IsVisibleInEditor() const override;
+		virtual bool IsVisibleInGame() const override;
 		virtual bool IsLocked() const override;
 		virtual bool IsTransient() const override;
 		virtual bool IsReadOnly() const override;
@@ -57,7 +58,8 @@ namespace WorldHierarchy
 		virtual bool CanSave() const override;
 		virtual bool HasLightingControls() const override;
 		virtual bool HasLockControls() const override;
-		virtual bool HasVisibilityControls() const override;
+		virtual bool HasEditorVisibilityControls() const override;
+		virtual bool HasGameVisibilityControls() const override;
 		virtual bool HasColorButtonControls() const override;
 		virtual bool HasKismet() const override;
 
@@ -71,9 +73,12 @@ namespace WorldHierarchy
 		virtual FLinearColor GetDrawColor() const override;
 		virtual void SetDrawColor(const FLinearColor& Color) override;
 
-		virtual void OnToggleVisibility() override;
-		virtual void OnShowOnlySelected() override;
-		virtual void OnShowAllButSelected() override;
+		virtual void OnToggleEditorVisibility() override;
+		virtual void OnShowInEditorOnlySelected() override;
+		virtual void OnShowInEditorAllButSelected() override;
+		virtual void OnToggleGameVisibility() override;
+		virtual void OnShowInGameOnlySelected() override;
+		virtual void OnShowInGameAllButSelected() override;
 		virtual void PopulateLevelModelList(FLevelModelList& InModelList) override;
 		virtual void OnToggleLightingScenario() override;
 		virtual void OnToggleLock() override;
@@ -93,7 +98,8 @@ namespace WorldHierarchy
 		virtual FLevelModelTreeItem* GetAsLevelModelTreeItem() const override { return const_cast<FLevelModelTreeItem*>(this); }
 		virtual FFolderTreeItem* GetAsFolderTreeItem() const override { return nullptr; }
 
-		virtual void SetVisible(bool bVisible) override;
+		virtual void SetVisibleInEditor(bool bVisible) override;
+		virtual void SetVisibleInGame(bool bVisible) override;
 		virtual void SetLocked(bool bLocked) override;
 
 	public:
@@ -108,6 +114,10 @@ namespace WorldHierarchy
 	private:
 
 		FWorldTreeItemID ID;
+
+		void SetSelectedLevelsToPopulatedList();
+
+		bool IsPersistentLevel() const;
 	};
 
 
@@ -127,11 +137,9 @@ namespace WorldHierarchy
 		virtual FString GetDisplayString() const override;
 
 		virtual FText GetToolTipText() const override;
-
 		virtual FText GetLockToolTipText() const override;
-
-		virtual FText GetVisibilityToolTipText() const override;
-
+		virtual FText GetEditorVisibilityToolTipText() const override;
+		virtual FText GetGameVisibilityToolTipText() const override;
 		virtual FText GetSaveToolTipText() const override;
 
 		virtual FWorldTreeItemID GetParentID() const override;
@@ -144,18 +152,23 @@ namespace WorldHierarchy
 
 		virtual int32 GetSortPriority() const override { return FolderSortPriority; }
 
-		virtual bool IsVisible() const override;
+		virtual bool IsVisibleInEditor() const override;
+		virtual bool IsVisibleInGame() const override;
 		virtual bool IsLocked() const override;
 
 		virtual bool CanSave() const override;
 		virtual bool HasLockControls() const override;
-		virtual bool HasVisibilityControls() const override;
+		virtual bool HasEditorVisibilityControls() const override;
+		virtual bool HasGameVisibilityControls() const override;
 		virtual bool HasValidPackage() const override;
 		virtual bool IsDirty() const override;
 
-		virtual void OnToggleVisibility() override;
-		virtual void OnShowOnlySelected() override;
-		virtual void OnShowAllButSelected() override;
+		virtual void OnToggleEditorVisibility() override;
+		virtual void OnShowInEditorOnlySelected() override;
+		virtual void OnShowInEditorAllButSelected() override;
+		virtual void OnToggleGameVisibility() override;
+		virtual void OnShowInGameOnlySelected() override;
+		virtual void OnShowInGameAllButSelected() override;
 		virtual void PopulateLevelModelList(FLevelModelList& InModelList) override;
 		virtual void OnToggleLock() override;
 		virtual void OnLockOnlySelected() override;
@@ -177,7 +190,8 @@ namespace WorldHierarchy
 		FName GetFullPath() const { return Path; }
 		FName GetLeafName() const { return LeafName; }
 
-		virtual void SetVisible(bool bVisible) override;
+		virtual void SetVisibleInEditor(bool bVisible) override;
+		virtual void SetVisibleInGame(bool bVisible) override;
 		virtual void SetLocked(bool bLocked) override;
 	
 	public:
@@ -186,7 +200,10 @@ namespace WorldHierarchy
 		virtual void OnDrop(const FDragDropEvent& DragEvent, TSharedRef<SWorldHierarchyImpl> Hierarchy) override;
 
 	private:
+		
 		bool IsAnyChildLoaded() const;
+		
+		FLevelModelList SetSelectionToFolderChildren();
 
 	private:
 		FName Path;
diff --git a/Engine/Source/Editor/WorldBrowser/Public/IWorldTreeItem.h b/Engine/Source/Editor/WorldBrowser/Public/IWorldTreeItem.h
--- a/Engine/Source/Editor/WorldBrowser/Public/IWorldTreeItem.h
+++ b/Engine/Source/Editor/WorldBrowser/Public/IWorldTreeItem.h
@@ -165,7 +165,10 @@ namespace WorldHierarchy
 		virtual FText GetLockToolTipText() const = 0;
 
 		/** Gets the tooltip for this item's visibility icon */
-		virtual FText GetVisibilityToolTipText() const = 0;
+		virtual FText GetEditorVisibilityToolTipText() const = 0;
+		
+		/** Gets the tooltip for this item's visibility icon */
+		virtual FText GetGameVisibilityToolTipText() const = 0;
 
 		/** Gets the tooltip text for this item's save icon */
 		virtual FText GetSaveToolTipText() const = 0;
@@ -201,7 +204,8 @@ namespace WorldHierarchy
 
 	public:
 
-		virtual bool IsVisible() const = 0;
+		virtual bool IsVisibleInEditor() const = 0;
+		virtual bool IsVisibleInGame() const = 0;
 		virtual bool IsLocked() const = 0;
 
 		/** Can this item be saved? */
@@ -214,7 +218,10 @@ namespace WorldHierarchy
 		virtual bool HasLockControls() const { return true; }
 
 		/** Can visibility on this item be toggled? */
-		virtual bool HasVisibilityControls() const { return true; }
+		virtual bool HasEditorVisibilityControls() const { return true; }
+		
+		/** Can visibility on this item be toggled? */
+		virtual bool HasGameVisibilityControls() const { return true; }
 
 		/** Does this item have color button controls? */
 		virtual bool HasColorButtonControls() const { return false; }
@@ -255,9 +262,12 @@ namespace WorldHierarchy
 		virtual const FSlateBrush* GetHierarchyItemBrush() const { return nullptr; }
 		virtual float GetHierarchyItemBrushWidth() const { return 7.0f; }
 
-		virtual void OnToggleVisibility() {}
-		virtual void OnShowOnlySelected() {}
-		virtual void OnShowAllButSelected() {}
+		virtual void OnToggleEditorVisibility() {}
+		virtual void OnShowInEditorOnlySelected() {}
+		virtual void OnShowInEditorAllButSelected() {}
+		virtual void OnToggleGameVisibility() {}
+		virtual void OnShowInGameOnlySelected() {}
+		virtual void OnShowInGameAllButSelected() {}
 		virtual void PopulateLevelModelList(FLevelModelList& InModelList) {}
 		virtual void OnToggleLightingScenario() {}
 		virtual void OnLockOnlySelected() {}
@@ -272,8 +282,10 @@ namespace WorldHierarchy
 		/** Generates a context menu option for this item if and only if it's the only item selected in the hierarchy */
 		virtual void GenerateContextMenu(UToolMenu* Menu, const SWorldHierarchyImpl& Hierarchy) {}
 
-		/** Sets the item's visible status */
-		virtual void SetVisible(bool bVisible) = 0;
+		/** Sets the item's editor visible status */
+		virtual void SetVisibleInEditor(bool bVisible) = 0;
+		/** Sets the item's game visible status */
+		virtual void SetVisibleInGame(bool bVisible) = 0;
 
 		/** Sets the item's locked status */
 		virtual void SetLocked(bool bLocked) = 0;
diff --git a/Engine/Source/Editor/WorldBrowser/Public/WorldHierarchyColumns.h b/Engine/Source/Editor/WorldBrowser/Public/WorldHierarchyColumns.h
--- a/Engine/Source/Editor/WorldBrowser/Public/WorldHierarchyColumns.h
+++ b/Engine/Source/Editor/WorldBrowser/Public/WorldHierarchyColumns.h
@@ -9,10 +9,11 @@ namespace UE::WorldHierarchy::HierarchyColumns
 	/** IDs for list columns */
 	WORLDBROWSER_API extern const FName ColumnID_LevelLabel;
 	WORLDBROWSER_API extern const FName ColumnID_EditorVisibility;
+	WORLDBROWSER_API extern const FName ColumnID_GameVisibility;
 	WORLDBROWSER_API extern const FName ColumnID_LightingScenario;
 	WORLDBROWSER_API extern const FName ColumnID_Lock;
 	WORLDBROWSER_API extern const FName ColumnID_SCCStatus;
 	WORLDBROWSER_API extern const FName ColumnID_Save;
 	WORLDBROWSER_API extern const FName ColumnID_Color;
 	WORLDBROWSER_API extern const FName ColumnID_Kismet;
-}
\ No newline at end of file
+}
