Save column visibility state so API user does not have to reapply it constantly when the world changes.
Also fixes failing ensure when SetColumnVisible or IsColumnVisible is called due to WeakWorldHierarchyImpl variable never being assigned to.

diff --git a/Engine/Source/Editor/WorldBrowser/Private/Misc/ColumnVisibilityState.cpp b/Engine/Source/Editor/WorldBrowser/Private/Misc/ColumnVisibilityState.cpp
new file mode 100644
--- /dev/null
+++ b/Engine/Source/Editor/WorldBrowser/Private/Misc/ColumnVisibilityState.cpp
@@ -0,0 +1,45 @@
+﻿// Copyright Epic Games, Inc. All Rights Reserved.
+
+#include "ColumnVisibilityState.h"
+
+#include "SWorldHierarchyImpl.h"
+
+namespace UE::WorldBrowser
+{
+	void FColumnVisibilityState::ApplyStateTo(SWorldHierarchyImpl& Widget) const
+	{
+		for (const TPair<FName, bool>& VisibilityState : VisibilityStates)
+		{
+			Widget.SetColumnVisible(VisibilityState.Key, VisibilityState.Value);
+		}
+	}
+
+	void FColumnVisibilityState::CaptureState(const SWorldHierarchyImpl& Widget)
+	{
+		for (TPair<FName, bool>& VisibilityState : VisibilityStates)
+		{
+			VisibilityState.Value = Widget.IsColumnVisible(VisibilityState.Key);
+		}
+	}
+
+	bool FColumnVisibilityState::IsColumnVisible(FName ColumnId) const
+	{
+		const bool* bIsVisible = VisibilityStates.Find(ColumnId);
+		return bIsVisible && *bIsVisible;
+	}
+
+	bool FColumnVisibilityState::SetVisibilityIfValidColumn(FName ColumnId, bool bShouldBeVisible)
+	{
+		if (ColumnId == WorldHierarchy::HierarchyColumns::ColumnID_LevelLabel)
+		{
+			return false;
+		}
+
+		if (bool* bIsVisible = VisibilityStates.Find(ColumnId))
+		{
+			*bIsVisible = bShouldBeVisible;
+			return true;
+		}
+		return false;
+	}
+}
diff --git a/Engine/Source/Editor/WorldBrowser/Private/Misc/ColumnVisibilityState.h b/Engine/Source/Editor/WorldBrowser/Private/Misc/ColumnVisibilityState.h
new file mode 100644
--- /dev/null
+++ b/Engine/Source/Editor/WorldBrowser/Private/Misc/ColumnVisibilityState.h
@@ -0,0 +1,48 @@
+﻿// Copyright Epic Games, Inc. All Rights Reserved.
+
+#pragma once
+
+#include "Containers/Map.h"
+#include "UObject/NameTypes.h"
+#include "WorldHierarchyColumns.h"
+
+class SWorldHierarchyImpl;
+
+namespace UE::WorldBrowser
+{
+	/**
+	 * Saves the column visibilities that SWorldHierarchyImpls should have as they are re-created.
+	 * Remember: SWorldHierarchyImpls is recreated every time a UWorld is opened. This class saves the column visibility states, so it can be restored.
+	 */
+	class FColumnVisibilityState
+	{
+	public:
+
+		/** Makes Widget show the columns saved in this object. */
+		void ApplyStateTo(SWorldHierarchyImpl& Widget) const;
+		
+		/** Saves the column visibility state from Widget into this object.*/
+		void CaptureState(const SWorldHierarchyImpl& Widget);
+
+		/** @return Whether the column is known and supposed to be visible */
+		bool IsColumnVisible(FName ColumnId) const;
+		
+		/** @return Whether ColumnId was a known column and was set to bShouldBeVisible. */
+		bool SetVisibilityIfValidColumn(FName ColumnId, bool bShouldBeVisible);
+
+	private:
+
+		TMap<FName, bool> VisibilityStates
+		{
+			{ WorldHierarchy::HierarchyColumns::ColumnID_LevelLabel, true },
+			{ WorldHierarchy::HierarchyColumns::ColumnID_EditorVisibility, true },
+			{ WorldHierarchy::HierarchyColumns::ColumnID_GameVisibility, false }, // Hidden by default
+			{ WorldHierarchy::HierarchyColumns::ColumnID_LightingScenario, true },
+			{ WorldHierarchy::HierarchyColumns::ColumnID_Lock, true },
+			{ WorldHierarchy::HierarchyColumns::ColumnID_Save, true },
+			{ WorldHierarchy::HierarchyColumns::ColumnID_Color, true },
+			{ WorldHierarchy::HierarchyColumns::ColumnID_Kismet, true },
+		};
+	};
+}
+
diff --git a/Engine/Source/Editor/WorldBrowser/Private/SWorldHierarchy.cpp b/Engine/Source/Editor/WorldBrowser/Private/SWorldHierarchy.cpp
--- a/Engine/Source/Editor/WorldBrowser/Private/SWorldHierarchy.cpp
+++ b/Engine/Source/Editor/WorldBrowser/Private/SWorldHierarchy.cpp
@@ -52,16 +52,20 @@ void SWorldHierarchy::Construct(const FArguments& InArgs)
 bool SWorldHierarchy::IsColumnVisible(FName Column) const
 {
 	const TSharedPtr<SWorldHierarchyImpl> WidgetPin = WeakWorldHierarchyImpl.Pin();
-	// This widget is destroyed during world construction
-	return ensureMsgf(WidgetPin, TEXT("Did you call this during world destruction?"))
-		&& WidgetPin->IsColumnVisible(Column);
+	return WidgetPin
+		? WidgetPin->IsColumnVisible(Column)
+		: VisibilityState.IsColumnVisible(Column);
 }
 
-void SWorldHierarchy::SetColumnVisible(FName Column, bool bVisible) const
+void SWorldHierarchy::SetColumnVisible(FName Column, bool bVisible)
 {
-	if (const TSharedPtr<SWorldHierarchyImpl> WidgetPin = WeakWorldHierarchyImpl.Pin();
-		// This widget is destroyed during world construction
-		ensureMsgf(WidgetPin, TEXT("Did you call this during world destruction?")))
+	const bool bIsKnownColumn = VisibilityState.SetVisibilityIfValidColumn(Column, bVisible);
+	if (!bIsKnownColumn)
+	{
+		return;
+	}
+	
+	if (const TSharedPtr<SWorldHierarchyImpl> WidgetPin = WeakWorldHierarchyImpl.Pin())
 	{
 		WidgetPin->SetColumnVisible(Column, bVisible);
 	}
@@ -69,12 +73,17 @@ void SWorldHierarchy::SetColumnVisible(FName Column, bool bVisible) const
 
 void SWorldHierarchy::OnBrowseWorld(UWorld* InWorld)
 {
+	// If the widget is about to be destroyed, capture the widget's column visibility state, so it can be re-applied later.
+	if (const TSharedPtr<SWorldHierarchyImpl> Widget = WeakWorldHierarchyImpl.Pin())
+	{
+		VisibilityState.CaptureState(*Widget);
+	}
+	
 	// Remove all binding to an old world
 	ChildSlot
 	[
 		SNullWidget::NullWidget
 	];
-
 	WorldModel = nullptr;
 		
 	// Bind to a new world
@@ -100,6 +109,7 @@ void SWorldHierarchy::OnBrowseWorld(UWorld* InWorld)
 			FWorldBrowserModule& WorldBrowserModule = FModuleManager::GetModuleChecked<FWorldBrowserModule>("WorldBrowser");
 			WorldModel = WorldBrowserModule.SharedWorldModel(InWorld);
 
+			TSharedPtr<SWorldHierarchyImpl> HierarchyImpl;
 			ChildSlot
 			[
 				SNew(SBorder)
@@ -162,11 +172,14 @@ void SWorldHierarchy::OnBrowseWorld(UWorld* InWorld)
 					.FillHeight(1.f)
 					.Padding(0.f, 4.f, 0.f ,0.f)
 					[
-						SNew(SWorldHierarchyImpl)
+						SAssignNew(HierarchyImpl, SWorldHierarchyImpl)
 						.InWorldModel(WorldModel)
 					]
 				]
 			];
+
+			WeakWorldHierarchyImpl = HierarchyImpl;
+			VisibilityState.ApplyStateTo(*HierarchyImpl);
 		}
 	}
 }
diff --git a/Engine/Source/Editor/WorldBrowser/Private/SWorldHierarchy.h b/Engine/Source/Editor/WorldBrowser/Private/SWorldHierarchy.h
--- a/Engine/Source/Editor/WorldBrowser/Private/SWorldHierarchy.h
+++ b/Engine/Source/Editor/WorldBrowser/Private/SWorldHierarchy.h
@@ -5,6 +5,7 @@
 #include "IWorldHierarchy.h"
 #include "Layout/Visibility.h"
 #include "Input/Reply.h"
+#include "Misc/ColumnVisibilityState.h"
 #include "Widgets/SWidget.h"
 #include "Widgets/DeclarativeSyntaxSupport.h"
 #include "Widgets/SCompoundWidget.h"
@@ -31,7 +32,7 @@ class SWorldHierarchy
 	//~ Begin IWorldHierarchy Interface
 	virtual TSharedRef<SWidget> GetWidget() override { return SharedThis(this); }
 	virtual bool IsColumnVisible(FName Column) const override;
-	virtual void SetColumnVisible(FName Column, bool bVisible) const override;
+	virtual void SetColumnVisible(FName Column, bool bVisible) override;
 	//~ End IWorldHierarchy Interface
 
 private:
@@ -62,4 +63,7 @@ private:
 	 * That is why this must be a weak pointer instead of a strong pointer.
 	 */
 	TWeakPtr<SWorldHierarchyImpl> WeakWorldHierarchyImpl;
+
+	/** Tracks which columns are supposed to be visible. */
+	UE::WorldBrowser::FColumnVisibilityState VisibilityState;
 };
diff --git a/Engine/Source/Editor/WorldBrowser/Public/IWorldHierarchy.h b/Engine/Source/Editor/WorldBrowser/Public/IWorldHierarchy.h
--- a/Engine/Source/Editor/WorldBrowser/Public/IWorldHierarchy.h
+++ b/Engine/Source/Editor/WorldBrowser/Public/IWorldHierarchy.h
@@ -26,7 +26,7 @@ namespace UE::WorldHierarchy
 		 * Sets whether Column is visible in the UI.
 		 * @see WorldHierarchyColumns.h for named columns.
 		 */
-		virtual void SetColumnVisible(FName Column, bool bVisible) const = 0;
+		virtual void SetColumnVisible(FName Column, bool bVisible) = 0;
 
 		virtual ~IWorldHierarchy() = default;
 	};
