World Browser now saves the columns that were visible

diff --git a/Engine/Source/Editor/WorldBrowser/Private/Misc/ColumnVisibilityState.cpp b/Engine/Source/Editor/WorldBrowser/Private/Misc/ColumnVisibilityState.cpp
deleted file mode 100644
diff --git a/Engine/Source/Editor/WorldBrowser/Private/Misc/ColumnVisibilityState.h b/Engine/Source/Editor/WorldBrowser/Private/Misc/ColumnVisibilityState.h
deleted file mode 100644
diff --git a/Engine/Source/Editor/WorldBrowser/Private/SWorldHierarchy.cpp b/Engine/Source/Editor/WorldBrowser/Private/SWorldHierarchy.cpp
--- a/Engine/Source/Editor/WorldBrowser/Private/SWorldHierarchy.cpp
+++ b/Engine/Source/Editor/WorldBrowser/Private/SWorldHierarchy.cpp
@@ -28,6 +28,7 @@
 #include "DragAndDrop/AssetDragDropOp.h"
 
 #include "SWorldHierarchyImpl.h"
+#include "WorldBrowserConfig.h"
 
 #define LOCTEXT_NAMESPACE "WorldBrowser"
 
@@ -54,31 +55,24 @@ bool SWorldHierarchy::IsColumnVisible(FName Column) const
 	const TSharedPtr<SWorldHierarchyImpl> WidgetPin = WeakWorldHierarchyImpl.Pin();
 	return WidgetPin
 		? WidgetPin->IsColumnVisible(Column)
-		: VisibilityState.IsColumnVisible(Column);
+		: SWorldHierarchyImpl::IsVisibleInConfig(Column);
 }
 
 void SWorldHierarchy::SetColumnVisible(FName Column, bool bVisible)
 {
-	const bool bIsKnownColumn = VisibilityState.SetVisibilityIfValidColumn(Column, bVisible);
-	if (!bIsKnownColumn)
-	{
-		return;
-	}
-	
 	if (const TSharedPtr<SWorldHierarchyImpl> WidgetPin = WeakWorldHierarchyImpl.Pin())
 	{
 		WidgetPin->SetColumnVisible(Column, bVisible);
 	}
+	else
+	{
+		// By documentation, SetColumnVisible does not save into the config. Only direct user action should do that.
+		SWorldHierarchyImpl::SetWillBeVisibleInConfigTransient(Column, bVisible);
+	}
 }
 
 void SWorldHierarchy::OnBrowseWorld(UWorld* InWorld)
 {
-	// If the widget is about to be destroyed, capture the widget's column visibility state, so it can be re-applied later.
-	if (const TSharedPtr<SWorldHierarchyImpl> Widget = WeakWorldHierarchyImpl.Pin())
-	{
-		VisibilityState.CaptureState(*Widget);
-	}
-	
 	// Remove all binding to an old world
 	ChildSlot
 	[
@@ -109,7 +103,6 @@ void SWorldHierarchy::OnBrowseWorld(UWorld* InWorld)
 			FWorldBrowserModule& WorldBrowserModule = FModuleManager::GetModuleChecked<FWorldBrowserModule>("WorldBrowser");
 			WorldModel = WorldBrowserModule.SharedWorldModel(InWorld);
 
-			TSharedPtr<SWorldHierarchyImpl> HierarchyImpl;
 			ChildSlot
 			[
 				SNew(SBorder)
@@ -172,14 +165,11 @@ void SWorldHierarchy::OnBrowseWorld(UWorld* InWorld)
 					.FillHeight(1.f)
 					.Padding(0.f, 4.f, 0.f ,0.f)
 					[
-						SAssignNew(HierarchyImpl, SWorldHierarchyImpl)
+						SAssignNew(WeakWorldHierarchyImpl, SWorldHierarchyImpl)
 						.InWorldModel(WorldModel)
 					]
 				]
 			];
-
-			WeakWorldHierarchyImpl = HierarchyImpl;
-			VisibilityState.ApplyStateTo(*HierarchyImpl);
 		}
 	}
 }
diff --git a/Engine/Source/Editor/WorldBrowser/Private/SWorldHierarchy.h b/Engine/Source/Editor/WorldBrowser/Private/SWorldHierarchy.h
--- a/Engine/Source/Editor/WorldBrowser/Private/SWorldHierarchy.h
+++ b/Engine/Source/Editor/WorldBrowser/Private/SWorldHierarchy.h
@@ -5,7 +5,6 @@
 #include "IWorldHierarchy.h"
 #include "Layout/Visibility.h"
 #include "Input/Reply.h"
-#include "Misc/ColumnVisibilityState.h"
 #include "Widgets/SWidget.h"
 #include "Widgets/DeclarativeSyntaxSupport.h"
 #include "Widgets/SCompoundWidget.h"
@@ -63,7 +62,4 @@ private:
 	 * That is why this must be a weak pointer instead of a strong pointer.
 	 */
 	TWeakPtr<SWorldHierarchyImpl> WeakWorldHierarchyImpl;
-
-	/** Tracks which columns are supposed to be visible. */
-	UE::WorldBrowser::FColumnVisibilityState VisibilityState;
 };
diff --git a/Engine/Source/Editor/WorldBrowser/Private/SWorldHierarchyImpl.cpp b/Engine/Source/Editor/WorldBrowser/Private/SWorldHierarchyImpl.cpp
--- a/Engine/Source/Editor/WorldBrowser/Private/SWorldHierarchyImpl.cpp
+++ b/Engine/Source/Editor/WorldBrowser/Private/SWorldHierarchyImpl.cpp
@@ -24,7 +24,9 @@
 #include "WorldTreeItemTypes.h"
 #include "LevelFolders.h"
 #include "ScopedTransaction.h"
+#include "Templates/UnrealTemplate.h"
 #include "Editor.h"
+#include "WorldBrowserConfig.h"
 #include "WorldBrowserStyle.h"
 #include "WorldHierarchyColumns.h"
 #include "RevisionControlStyle/RevisionControlStyle.h"
@@ -288,9 +290,10 @@ TSharedRef<SHeaderRow> SWorldHierarchyImpl::CreateHeaderRow()
 {
 	using namespace UE::WorldHierarchy;
 	constexpr float IconWidth = 24.f;
-	
+
 	TSharedRef<SHeaderRow> HeaderRow = SAssignNew(HeaderRowWidget, SHeaderRow)
 		.CanSelectGeneratedColumn(true) // Lets the user choose which columns should be shown
+		.OnHiddenColumnsListChanged(this, &SWorldHierarchyImpl::SaveColumnVisibilitiesIntoConfig)
 
 		/** Level visibility column */
 		+ SHeaderRow::Column(HierarchyColumns::ColumnID_EditorVisibility)
@@ -407,12 +410,64 @@ TSharedRef<SHeaderRow> SWorldHierarchyImpl::CreateHeaderRow()
 			.Image(FAppStyle::Get().GetBrush("Level.ColorIcon"))
 			.ToolTipText(LOCTEXT("Column.Color.Tooltip", "Color used for visualization of Level"))
 		];
-
-	// By default, hide this column; it is only used in certain workflows, like Virtual Production. The user can manually show this column.
-	SetColumnVisible(HierarchyColumns::ColumnID_GameVisibility, false);
+	
+	UWorldBrowserConfig::Initialize();
+	UWorldBrowserConfig* Config = UWorldBrowserConfig::Get();
+	for (const TPair<FName, bool>& VisibilityPair : Config->ColumnConfig.ColumnVisibilities)
+	{
+		const FName ColumnId = VisibilityPair.Key;
+		if (!IsRequiredColumn(ColumnId))
+		{
+			SetColumnVisible(ColumnId, VisibilityPair.Value);
+		}
+	}
+	
 	return HeaderRow;
 }
 
+void SWorldHierarchyImpl::SaveColumnVisibilitiesIntoConfig()
+{
+	if (bIsProgrammaticallyChangingColumnVisibility)
+	{
+		return;
+	}
+	
+	UWorldBrowserConfig::Initialize();
+	UWorldBrowserConfig* Config = UWorldBrowserConfig::Get();
+
+	TMap<FName, bool>& Visibilities = Config->ColumnConfig.ColumnVisibilities;
+	Visibilities.Empty();
+	for (const SHeaderRow::FColumn& Column : HeaderRowWidget->GetColumns())
+	{
+		const FName ColumnId = Column.ColumnId;
+		if (!IsRequiredColumn(ColumnId))
+		{
+			Visibilities.Add(ColumnId, IsColumnVisible(ColumnId));
+		}
+	}
+	
+	Config->SaveEditorConfig();
+}
+
+bool SWorldHierarchyImpl::IsRequiredColumn(FName ColumnId)
+{
+	return ColumnId == UE::WorldHierarchy::HierarchyColumns::ColumnID_LevelLabel;
+}
+
+bool SWorldHierarchyImpl::IsKnownColumn(FName ColumnId)
+{
+	using namespace UE::WorldHierarchy;
+	return ColumnId == HierarchyColumns::ColumnID_EditorVisibility
+		|| ColumnId == HierarchyColumns::ColumnID_GameVisibility
+		|| ColumnId == HierarchyColumns::ColumnID_LevelLabel
+		|| ColumnId == HierarchyColumns::ColumnID_LightingScenario
+		|| ColumnId == HierarchyColumns::ColumnID_Lock
+		|| ColumnId == HierarchyColumns::ColumnID_SCCStatus
+		|| ColumnId == HierarchyColumns::ColumnID_Save
+		|| ColumnId == HierarchyColumns::ColumnID_Color
+		|| ColumnId == HierarchyColumns::ColumnID_Kismet;
+}
+
 TSharedRef<ITableRow> SWorldHierarchyImpl::GenerateTreeRow(WorldHierarchy::FWorldTreeItemPtr Item, const TSharedRef<STableViewBase>& OwnerTable)
 {
 	check(Item.IsValid());
@@ -747,9 +802,31 @@ bool SWorldHierarchyImpl::IsColumnVisible(FName ColumnId) const
 	return HeaderRowWidget->IsColumnVisible(ColumnId);
 }
 
-void SWorldHierarchyImpl::SetColumnVisible(FName ColumnId, bool bVisible) const
+void SWorldHierarchyImpl::SetColumnVisible(FName ColumnId, bool bVisible)
 {
-	HeaderRowWidget->SetShowGeneratedColumn(ColumnId, bVisible);
+	if (IsKnownColumn(ColumnId))
+	{
+		TGuardValue Guard(bIsProgrammaticallyChangingColumnVisibility, true);
+		HeaderRowWidget->SetShowGeneratedColumn(ColumnId, bVisible);
+	}
+}
+
+bool SWorldHierarchyImpl::IsVisibleInConfig(FName ColumnId)
+{
+	UWorldBrowserConfig::Initialize();
+	UWorldBrowserConfig* Config = UWorldBrowserConfig::Get();
+	const bool* bIsVisible = Config->ColumnConfig.ColumnVisibilities.Find(ColumnId);
+	return IsRequiredColumn(ColumnId) || !bIsVisible || *bIsVisible;
+}
+
+void SWorldHierarchyImpl::SetWillBeVisibleInConfigTransient(FName ColumnId, bool bIsVisible)
+{
+	if (!IsRequiredColumn(ColumnId) && IsKnownColumn(ColumnId))
+	{
+		UWorldBrowserConfig::Initialize();
+		UWorldBrowserConfig* Config = UWorldBrowserConfig::Get();
+		Config->ColumnConfig.ColumnVisibilities.Add(ColumnId, bIsVisible);
+	}
 }
 
 void SWorldHierarchyImpl::OnTreeItemScrolledIntoView( WorldHierarchy::FWorldTreeItemPtr Item, const TSharedPtr<ITableRow>& Widget )
diff --git a/Engine/Source/Editor/WorldBrowser/Private/SWorldHierarchyImpl.h b/Engine/Source/Editor/WorldBrowser/Private/SWorldHierarchyImpl.h
--- a/Engine/Source/Editor/WorldBrowser/Private/SWorldHierarchyImpl.h
+++ b/Engine/Source/Editor/WorldBrowser/Private/SWorldHierarchyImpl.h
@@ -108,7 +108,12 @@ public:
 	 * Sets whether Column is visible in the UI.
 	 * @see WorldHierarchyColumns.h for named columns.
 	 */
-	void SetColumnVisible(FName ColumnId, bool bVisible) const;
+	void SetColumnVisible(FName ColumnId, bool bVisible);
+
+	/** @return Whether the column will be visible according to the config. */
+	static bool IsVisibleInConfig(FName ColumnId);
+	/** Sets whether the next widget created from now will have ColumnId visible but does not save it into the config. */
+	static void SetWillBeVisibleInConfigTransient(FName ColumnId, bool bIsVisible);
 
 public:
 	//~ FEditorUndoClient
@@ -121,6 +126,15 @@ private:
 
 	/** Creates the header row for the level hierarchy. */
 	TSharedRef<SHeaderRow> CreateHeaderRow();
+
+	/** Invokes when the column visibility changes. */
+	void SaveColumnVisibilitiesIntoConfig();
+
+	/** @return Whether the given column is always shown. */
+	static bool IsRequiredColumn(FName ColumnId);
+	
+	/** @return Whether the column is used by this UI */
+	static bool IsKnownColumn(FName ColumnId);
 	
 	/** Creates an item for the tree view */
 	TSharedRef<ITableRow> GenerateTreeRow(WorldHierarchy::FWorldTreeItemPtr Item, const TSharedRef<STableViewBase>& OwnerTable);
@@ -354,6 +368,9 @@ private:
 	/** True if the items require sort */
 	bool bSortDirty;
 
+	/** Whether the column visibility is currently being changed via code (as opposed to via user selection in the UI). */
+	bool bIsProgrammaticallyChangingColumnVisibility = false;
+
 	/** Operations that are waiting to be resolved for items in the tree */
 	TArray<WorldHierarchy::FPendingWorldTreeOperation> PendingOperations;
 
diff --git a/Engine/Source/Editor/WorldBrowser/Private/WorldBrowserConfig.cpp b/Engine/Source/Editor/WorldBrowser/Private/WorldBrowserConfig.cpp
new file mode 100644
--- /dev/null
+++ b/Engine/Source/Editor/WorldBrowser/Private/WorldBrowserConfig.cpp
@@ -0,0 +1,15 @@
+﻿// Copyright Epic Games, Inc. All Rights Reserved.
+
+#include "WorldBrowserConfig.h"
+
+TObjectPtr<UWorldBrowserConfig> UWorldBrowserConfig::Instance = nullptr;
+
+void UWorldBrowserConfig::Initialize()
+{
+	if(!Instance)
+	{
+		Instance = NewObject<UWorldBrowserConfig>(); 
+		Instance->AddToRoot();
+		Instance->LoadEditorConfig();
+	}
+}
\ No newline at end of file
diff --git a/Engine/Source/Editor/WorldBrowser/Private/WorldBrowserConfig.h b/Engine/Source/Editor/WorldBrowser/Private/WorldBrowserConfig.h
new file mode 100644
--- /dev/null
+++ b/Engine/Source/Editor/WorldBrowser/Private/WorldBrowserConfig.h
@@ -0,0 +1,41 @@
+﻿// Copyright Epic Games, Inc. All Rights Reserved.
+
+#pragma once
+
+#include "Containers/Map.h"
+#include "EditorConfigBase.h"
+#include "WorldHierarchyColumns.h"
+#include "UObject/ObjectPtr.h"
+#include "UObject/UObjectGlobals.h"
+
+#include "WorldBrowserConfig.generated.h"
+
+USTRUCT()
+struct FWorldBrowserColumnConfig
+{
+	GENERATED_BODY()
+
+	UPROPERTY()
+	TMap<FName, bool> ColumnVisibilities
+	{
+		// By default, hide this column; it is only used in certain workflows, like Virtual Production. The user can manually show this column.
+		{ UE::WorldHierarchy::HierarchyColumns::ColumnID_GameVisibility, false }
+	};
+};
+
+UCLASS(EditorConfig="WorldBrowser")
+class UWorldBrowserConfig : public UEditorConfigBase
+{
+	GENERATED_BODY()
+public:
+	
+	UPROPERTY(meta=(EditorConfig))
+	FWorldBrowserColumnConfig ColumnConfig;
+
+	static void Initialize();
+	static UWorldBrowserConfig* Get() { return Instance; }
+
+private:
+
+	static TObjectPtr<UWorldBrowserConfig> Instance;
+};
diff --git a/Engine/Source/Editor/WorldBrowser/Public/IWorldHierarchy.h b/Engine/Source/Editor/WorldBrowser/Public/IWorldHierarchy.h
--- a/Engine/Source/Editor/WorldBrowser/Public/IWorldHierarchy.h
+++ b/Engine/Source/Editor/WorldBrowser/Public/IWorldHierarchy.h
@@ -23,7 +23,7 @@ namespace UE::WorldHierarchy
 		 */
 		virtual bool IsColumnVisible(FName Column) const = 0;
 		/**
-		 * Sets whether Column is visible in the UI.
+		 * Sets whether Column is visible in the UI. Does not save this into the config though.
 		 * @see WorldHierarchyColumns.h for named columns.
 		 */
 		virtual void SetColumnVisible(FName Column, bool bVisible) = 0;
diff --git a/Engine/Source/Editor/WorldBrowser/WorldBrowser.Build.cs b/Engine/Source/Editor/WorldBrowser/WorldBrowser.Build.cs
--- a/Engine/Source/Editor/WorldBrowser/WorldBrowser.Build.cs
+++ b/Engine/Source/Editor/WorldBrowser/WorldBrowser.Build.cs
@@ -20,7 +20,8 @@ public class WorldBrowser : ModuleRules
 				"ApplicationCore",
                 "AppFramework",
                 "Core", 
-                "CoreUObject",
+                "CoreUObject", 
+                "EditorConfig",
                 "RenderCore",
                 "InputCore",
                 "Engine",
