Fix inconsistent loca key values

diff --git a/Engine/Source/Editor/WorldBrowser/Private/Tiles/WorldTileCollectionModel.cpp b/Engine/Source/Editor/WorldBrowser/Private/Tiles/WorldTileCollectionModel.cpp
--- a/Engine/Source/Editor/WorldBrowser/Private/Tiles/WorldTileCollectionModel.cpp
+++ b/Engine/Source/Editor/WorldBrowser/Private/Tiles/WorldTileCollectionModel.cpp
@@ -461,7 +461,7 @@ void FWorldTileCollectionModel::BuildHierarchyMenu(FMenuBuilder& InMenuBuilder,
 		// Visibility commands
 		InMenuBuilder.AddSubMenu( 
 			LOCTEXT("EditorVisibilityHeader", "Visibility in Editor"),
-			LOCTEXT("EditorVisibilitySubMenu_ToolTip", "Selected Level(s) visibility commands for editor"),
+			LOCTEXT("EditorVisibilitySubMenu_ToolTip", "Selected Level(s) visibility commands for editor worlds"),
 			FNewMenuDelegate::CreateSP(const_cast<FWorldTileCollectionModel*>(this), &FWorldTileCollectionModel::FillEditorVisibilitySubMenu ) );
 
 		// Game Visibility commands
