Add WORLDBROWSER_API back to private LevelCollectionModel.h header3rd party plugins may be depending on it.

diff --git a/Engine/Source/Editor/WorldBrowser/Private/LevelCollectionModel.h b/Engine/Source/Editor/WorldBrowser/Private/LevelCollectionModel.h
--- a/Engine/Source/Editor/WorldBrowser/Private/LevelCollectionModel.h
+++ b/Engine/Source/Editor/WorldBrowser/Private/LevelCollectionModel.h
@@ -29,7 +29,7 @@ enum class EBuildHierarchyMenuFlags
 ENUM_CLASS_FLAGS(EBuildHierarchyMenuFlags);
 
 /** Interface for non-UI presentation logic for a world. */
-class FLevelCollectionModel
+class WORLDBROWSER_API FLevelCollectionModel
 	: public TSharedFromThis<FLevelCollectionModel>	
 	, public FTickableEditorObject
 {
