[Viewport Toolbar] Copy two entries from main toolbar "Settings" to "Levels" tab menu

diff --git a/Engine/Source/Editor/WorldBrowser/Private/StreamingLevels/StreamingLevelCollectionModel.cpp b/Engine/Source/Editor/WorldBrowser/Private/StreamingLevels/StreamingLevelCollectionModel.cpp
--- a/Engine/Source/Editor/WorldBrowser/Private/StreamingLevels/StreamingLevelCollectionModel.cpp
+++ b/Engine/Source/Editor/WorldBrowser/Private/StreamingLevels/StreamingLevelCollectionModel.cpp
@@ -18,6 +18,8 @@
 
 #include "Interfaces/IMainFrameModule.h"
 #include "DesktopPlatformModule.h"
+#include "LevelEditor.h"
+#include "LevelEditorActions.h"
 #include "NewLevelDialogModule.h"
 
 #include "StreamingLevels/StreamingLevelCustomization.h"
@@ -396,6 +398,17 @@ void FStreamingLevelCollectionModel::CustomizeFileMainMenu(FMenuBuilder& InMenuB
 		InMenuBuilder.AddMenuEntry( Commands.World_MergeSelectedLevels );
 	}
 	InMenuBuilder.EndSection();
+
+	InMenuBuilder.BeginSection("Options", LOCTEXT("OptionsLabel", "Options"));
+	{
+		FLevelEditorModule& LevelEditorModule = FModuleManager::Get().LoadModuleChecked<FLevelEditorModule>("LevelEditor");
+		TSharedRef<FUICommandList> GlobalLevelCommandList = LevelEditorModule.GetGlobalLevelEditorActions();
+		InMenuBuilder.PushCommandList(GlobalLevelCommandList);
+
+		InMenuBuilder.AddMenuEntry(FLevelEditorCommands::Get().OnlyLoadVisibleInPIE);
+		InMenuBuilder.AddMenuEntry(FLevelEditorCommands::Get().LevelStreamingVolumePrevis);
+	}
+	InMenuBuilder.EndSection();
 }
 
 void FStreamingLevelCollectionModel::FillDefaultStreamingMethodSubMenu(FMenuBuilder& InMenuBuilder)
