Updated headers using UnrealCodeFixup to make sure dllstorage is on methods/staticvars instead of types. Part 5/n

diff --git a/Engine/Source/Editor/WorldBrowser/Private/LevelCollectionModel.h b/Engine/Source/Editor/WorldBrowser/Private/LevelCollectionModel.h
--- a/Engine/Source/Editor/WorldBrowser/Private/LevelCollectionModel.h
+++ b/Engine/Source/Editor/WorldBrowser/Private/LevelCollectionModel.h
@@ -13,6 +13,8 @@
 
 #include "Misc/FilterCollection.h"
 
+#define UE_API WORLDBROWSER_API
+
 class FMenuBuilder;
 class IDetailsView;
 class UEditorEngine;
@@ -29,7 +31,7 @@ enum class EBuildHierarchyMenuFlags
 ENUM_CLASS_FLAGS(EBuildHierarchyMenuFlags);
 
 /** Interface for non-UI presentation logic for a world. */
-class WORLDBROWSER_API FLevelCollectionModel
+class FLevelCollectionModel
 	: public TSharedFromThis<FLevelCollectionModel>	
 	, public FTickableEditorObject
 {
@@ -38,156 +40,156 @@ public:
 	DECLARE_EVENT_OneParam( FLevelCollectionModel, FOnNewItemAdded, TSharedPtr<FLevelModel>);
 	DECLARE_EVENT( FLevelCollectionModel, FSimpleEvent );
 	
-	FLevelCollectionModel();
-	virtual ~FLevelCollectionModel() override;
+	UE_API FLevelCollectionModel();
+	UE_API virtual ~FLevelCollectionModel() override;
 
 	/** FTickableEditorObject interface */
-	virtual void Tick( float DeltaTime ) override;
+	UE_API virtual void Tick( float DeltaTime ) override;
 	virtual ETickableTickType GetTickableTickType() const override { return ETickableTickType::Always; }
-	virtual TStatId GetStatId() const override;
+	UE_API virtual TStatId GetStatId() const override;
 	/** FTickableEditorObject interface */
 	
 	/**	@return	Whether level collection is read only now */
-	bool IsReadOnly() const;
+	UE_API bool IsReadOnly() const;
 	
 	/**	@return	Whether level collection is in PIE/SIE mode */
-	bool IsSimulating() const;
+	UE_API bool IsSimulating() const;
 
 	/**	@return	Current simulation world */
-	UWorld* GetSimulationWorld() const;
+	UE_API UWorld* GetSimulationWorld() const;
 
 	/**	@return	Current editor world */
 	UWorld* GetWorld(bool bEvenIfPendingKill = false) const { return CurrentWorld.Get(bEvenIfPendingKill); }
 
 	/** @return	Whether current world has world origin rebasing enabled */
-	bool IsOriginRebasingEnabled() const;
+	UE_API bool IsOriginRebasingEnabled() const;
 
 	/** Current world size  */
 	FIntPoint GetWorldSize() const { return WorldSize; }
 	
 	/**	@return	Root list of levels in hierarchy */
-	FLevelModelList& GetRootLevelList();
+	UE_API FLevelModelList& GetRootLevelList();
 
 	/**	@return	All level list managed by this level collection */
-	const FLevelModelList& GetAllLevels() const;
+	UE_API const FLevelModelList& GetAllLevels() const;
 
 	/**	@return	List of filtered levels */
-	const FLevelModelList& GetFilteredLevels() const;
+	UE_API const FLevelModelList& GetFilteredLevels() const;
 
 	/**	@return	Currently selected level list */
-	const FLevelModelList& GetSelectedLevels() const;
+	UE_API const FLevelModelList& GetSelectedLevels() const;
 
 	/** Adds a filter which restricts the Levels shown in UI */
-	void AddFilter(const TSharedRef<LevelFilter>& InFilter);
+	UE_API void AddFilter(const TSharedRef<LevelFilter>& InFilter);
 
 	/** Removes a filter which restricted the Levels shown in UI */
-	void RemoveFilter(const TSharedRef<LevelFilter>& InFilter);
+	UE_API void RemoveFilter(const TSharedRef<LevelFilter>& InFilter);
 
 	/**	@return	Whether level filtering is active now */
-	bool IsFilterActive() const;
+	UE_API bool IsFilterActive() const;
 	
 	/**	Iterates through level hierarchy with given Visitor */
-	void IterateHierarchy(FLevelModelVisitor& Visitor);
+	UE_API void IterateHierarchy(FLevelModelVisitor& Visitor);
 
 	/**	Sets selected level list */
-	void SetSelectedLevels(const FLevelModelList& InList);
+	UE_API void SetSelectedLevels(const FLevelModelList& InList);
 	
 	/**	Sets selection to a levels that is currently marked as selected in UWorld */
-	void SetSelectedLevelsFromWorld();
+	UE_API void SetSelectedLevelsFromWorld();
 
 	/**	@return	Found level model which represents specified level object */
-	TSharedPtr<FLevelModel> FindLevelModel(ULevel* InLevel) const;
+	UE_API TSharedPtr<FLevelModel> FindLevelModel(ULevel* InLevel) const;
 
 	/**	@return	Found level model with specified level package name */
-	TSharedPtr<FLevelModel> FindLevelModel(const FName& PackageName) const;
+	UE_API TSharedPtr<FLevelModel> FindLevelModel(const FName& PackageName) const;
 
 	/**	Hides level in editor worlds */
-	void HideLevelsInEditor(const FLevelModelList& InLevelList);
+	UE_API void HideLevelsInEditor(const FLevelModelList& InLevelList);
 
 	/**	Shows level in editor worlds */
-	void ShowLevelsInEditor(const FLevelModelList& InLevelList);
+	UE_API void ShowLevelsInEditor(const FLevelModelList& InLevelList);
 
 	/** Toggles the selected levels to a visible state; toggles all other levels to an invisible state */
-	void ShowInEditorOnlySelectedLevels();
+	UE_API void ShowInEditorOnlySelectedLevels();
 
 	/** Toggles the selected levels to an invisible state; toggles all other levels to a visible state */
-	void ShowInEditorAllButSelectedLevels();
+	UE_API void ShowInEditorAllButSelectedLevels();
 	
 	/**	Hides level in game worlds */
-	void HideLevelsInGame(const FLevelModelList& InLevelList);
+	UE_API void HideLevelsInGame(const FLevelModelList& InLevelList);
 
 	/**	Shows level in game worlds */
-	void ShowLevelsInGame(const FLevelModelList& InLevelList);
+	UE_API void ShowLevelsInGame(const FLevelModelList& InLevelList);
 
 	/** Toggles the selected levels to a visible state for game worlds; toggles all other levels to an invisible state */
-	void ShowInGameOnlySelectedLevels();
+	UE_API void ShowInGameOnlySelectedLevels();
 
 	/** Toggles the selected levels to an invisible state for game worlds; toggles all other levels to a visible state */
-	void ShowInGameAllButSelectedLevels();
+	UE_API void ShowInGameAllButSelectedLevels();
 
 	/**	Unlocks level in the world */
-	void UnlockLevels(const FLevelModelList& InLevelList);
+	UE_API void UnlockLevels(const FLevelModelList& InLevelList);
 	/**	Locks level in the world */
-	void LockLevels(const FLevelModelList& InLevelList);
+	UE_API void LockLevels(const FLevelModelList& InLevelList);
 
 	/** Toggles the selected levels to a locked state; toggles all other levels to an unlocked state */
-	void LockOnlySelectedLevels();
+	UE_API void LockOnlySelectedLevels();
 
 	/** Toggles the selected levels to an unlocked state; toggles all other levels to a locked state */
-	void LockAllButSelectedLevels();
+	UE_API void LockAllButSelectedLevels();
 
 	/**	Saves level to disk */
-	void SaveLevels(const FLevelModelList& InLevelList);
+	UE_API void SaveLevels(const FLevelModelList& InLevelList);
 
 	/**	Loads level from disk */
-	void LoadLevels(const FLevelModelList& InLevelList);
+	UE_API void LoadLevels(const FLevelModelList& InLevelList);
 	/**	Unloads levels from the editor */
-	virtual void UnloadLevels(const FLevelModelList& InLevelList);
+	UE_API virtual void UnloadLevels(const FLevelModelList& InLevelList);
 
 	/** Translate levels by specified delta */
-	virtual void TranslateLevels(const FLevelModelList& InLevelList, FVector2D InAbsoluteDelta, bool bSnapDelta = true);
+	UE_API virtual void TranslateLevels(const FLevelModelList& InLevelList, FVector2D InAbsoluteDelta, bool bSnapDelta = true);
 	
 	/** Snaps translation delta */
-	virtual FVector2D SnapTranslationDelta(const FLevelModelList& InLevelList, FVector2D InAbsoluteDelta, bool bBoundsSnapping, FVector2D::FReal SnappingValue);
+	UE_API virtual FVector2D SnapTranslationDelta(const FLevelModelList& InLevelList, FVector2D InAbsoluteDelta, bool bBoundsSnapping, FVector2D::FReal SnappingValue);
 
 	/**	Updates current translation delta, when user drags levels on minimap */
-	virtual void UpdateTranslationDelta(const FLevelModelList& InLevelList, FVector2D InTranslationDelta, bool bBoundsSnapping, FVector2D::FReal SnappingValue);
+	UE_API virtual void UpdateTranslationDelta(const FLevelModelList& InLevelList, FVector2D InTranslationDelta, bool bBoundsSnapping, FVector2D::FReal SnappingValue);
 
 	/** Attach levels as children to specified level */
-	void AssignParent(const FLevelModelList& InLevels, TSharedPtr<FLevelModel> InParent);
+	UE_API void AssignParent(const FLevelModelList& InLevels, TSharedPtr<FLevelModel> InParent);
 
 	/** Adds all levels in worlds represented by the supplied world list as sublevels */
-	virtual void AddExistingLevelsFromAssetData(const TArray<struct FAssetData>& WorldList);
+	UE_API virtual void AddExistingLevelsFromAssetData(const TArray<struct FAssetData>& WorldList);
 			
 	/**	Create drag drop operation for a selected level models */
-	virtual TSharedPtr<WorldHierarchy::FWorldBrowserDragDropOp> CreateDragDropOp() const;
+	UE_API virtual TSharedPtr<WorldHierarchy::FWorldBrowserDragDropOp> CreateDragDropOp() const;
 
 	/** Create a drag and drop operation for the specified level models */
-	virtual TSharedPtr<WorldHierarchy::FWorldBrowserDragDropOp> CreateDragDropOp(const FLevelModelList& InLevels) const;
+	UE_API virtual TSharedPtr<WorldHierarchy::FWorldBrowserDragDropOp> CreateDragDropOp(const FLevelModelList& InLevels) const;
 	
 	/**	@return	Whether specified level passes all filters */
-	virtual bool PassesAllFilters(const FLevelModel& InLevelModel) const;
+	UE_API virtual bool PassesAllFilters(const FLevelModel& InLevelModel) const;
 	
 	/**	Builds 'hierarchy' commands menu for a selected levels */
-	virtual void BuildHierarchyMenu(FMenuBuilder& InMenuBuilder, EBuildHierarchyMenuFlags Flags) const;
+	UE_API virtual void BuildHierarchyMenu(FMenuBuilder& InMenuBuilder, EBuildHierarchyMenuFlags Flags) const;
 	
 	/**	Customize 'File' section in main menu  */
-	virtual void CustomizeFileMainMenu(FMenuBuilder& InMenuBuilder) const;
+	UE_API virtual void CustomizeFileMainMenu(FMenuBuilder& InMenuBuilder) const;
 		
 	/**	@return	Player view in the PIE/Simulation world */
-	virtual bool GetPlayerView(FVector& Location, FRotator& Rotation) const;
+	UE_API virtual bool GetPlayerView(FVector& Location, FRotator& Rotation) const;
 
 	/**	@return	Observer view in the Editor/Similuation world */
-	virtual bool GetObserverView(FVector& Location, FRotator& Rotation) const;
+	UE_API virtual bool GetObserverView(FVector& Location, FRotator& Rotation) const;
 
 	/**	Compares 2 levels by Z order */
-	virtual bool CompareLevelsZOrder(TSharedPtr<FLevelModel> InA, TSharedPtr<FLevelModel> InB) const;
+	UE_API virtual bool CompareLevelsZOrder(TSharedPtr<FLevelModel> InA, TSharedPtr<FLevelModel> InB) const;
 
 	/**	Registers level details customizations */
-	virtual void RegisterDetailsCustomization(class FPropertyEditorModule& PropertyModule, TSharedPtr<class IDetailsView> InDetailsView);
+	UE_API virtual void RegisterDetailsCustomization(class FPropertyEditorModule& PropertyModule, TSharedPtr<class IDetailsView> InDetailsView);
 	
 	/**	Unregisters level details customizations */
-	virtual void UnregisterDetailsCustomization(class FPropertyEditorModule& PropertyModule, TSharedPtr<class IDetailsView> InDetailsView);
+	UE_API virtual void UnregisterDetailsCustomization(class FPropertyEditorModule& PropertyModule, TSharedPtr<class IDetailsView> InDetailsView);
 
 	/** @return	Whether this level collection model is a tile world */
 	virtual bool IsTileWorld() const { return false; };
@@ -196,308 +198,308 @@ public:
 	virtual bool HasFolderSupport() const { return false; }
 
 	/** Rebuilds levels collection */
-	void PopulateLevelsList();
+	UE_API void PopulateLevelsList();
 
 	/** Rebuilds the list of filtered Levels */
-	void PopulateFilteredLevelsList();
+	UE_API void PopulateFilteredLevelsList();
 
 	/**	Request to update levels cached information */
-	void RequestUpdateAllLevels();
+	UE_API void RequestUpdateAllLevels();
 	
 	/**	Request to redraw all levels */
-	void RequestRedrawAllLevels();
+	UE_API void RequestRedrawAllLevels();
 
 	/**	Updates all levels cached information */
-	void UpdateAllLevels();
+	UE_API void UpdateAllLevels();
 
 	/**	Redraws all levels */
-	void RedrawAllLevels();
+	UE_API void RedrawAllLevels();
 
 	/** Updates level actor count for all levels */
-	void UpdateLevelActorsCount();
+	UE_API void UpdateLevelActorsCount();
 
 	/** @return	whether exactly one level is selected */
-	bool IsOneLevelSelected() const;
+	UE_API bool IsOneLevelSelected() const;
 
 	/** @return	whether at least one level is selected */
-	bool AreAnyLevelsSelected() const;
+	UE_API bool AreAnyLevelsSelected() const;
 
 	/** @return wether all selected levels are user managed */
-	bool AreAllSelectedLevelsUserManaged() const;
+	UE_API bool AreAllSelectedLevelsUserManaged() const;
 
 	/** @return whether all the currently selected levels are loaded */
-	bool AreAllSelectedLevelsLoaded() const;
+	UE_API bool AreAllSelectedLevelsLoaded() const;
 
 	/** @return whether any of the currently selected levels is loaded */
-	bool AreAnySelectedLevelsLoaded() const;
+	UE_API bool AreAnySelectedLevelsLoaded() const;
 	
 	/** @return whether all the currently selected levels are unloaded */
-	bool AreAllSelectedLevelsUnloaded() const;
+	UE_API bool AreAllSelectedLevelsUnloaded() const;
 	
 	/** @return whether any of the currently selected levels is unloaded */
-	bool AreAnySelectedLevelsUnloaded() const;
+	UE_API bool AreAnySelectedLevelsUnloaded() const;
 
 	/** @return whether all the currently selected levels are editable */
-	bool AreAllSelectedLevelsEditable() const;
+	UE_API bool AreAllSelectedLevelsEditable() const;
 
 	/** @return whether all the currently selected levels are editable and not persistent */
-	bool AreAllSelectedLevelsEditableAndNotPersistent() const;
+	UE_API bool AreAllSelectedLevelsEditableAndNotPersistent() const;
 
 	/** @return whether all the currently selected levels are editable and visible*/
-	bool AreAllSelectedLevelsEditableAndVisible() const;
+	UE_API bool AreAllSelectedLevelsEditableAndVisible() const;
 
 	/** @return whether any of the currently selected levels is editable */
-	bool AreAnySelectedLevelsEditable() const;
+	UE_API bool AreAnySelectedLevelsEditable() const;
 
 	/** @return whether any of the currently selected levels is editable and visible*/
-	bool AreAnySelectedLevelsEditableAndVisible() const;
+	UE_API bool AreAnySelectedLevelsEditableAndVisible() const;
 
 	/** @return Whether game visibility of the selected levels is allowed to be changed. */
-	bool CanExecuteGameVisibilityCommandsForSelectedLevels() const;
+	UE_API bool CanExecuteGameVisibilityCommandsForSelectedLevels() const;
 	
 	/** @return Whether game visibility of levels is allowed to be changed. */
-	bool CanExecuteGameVisibilityCommands() const;
+	UE_API bool CanExecuteGameVisibilityCommands() const;
 	
 	/** @return whether currently only one level selected and it is editable */
-	bool IsSelectedLevelEditable() const;
+	UE_API bool IsSelectedLevelEditable() const;
 
 	/** @return whether currently only one level selected and a lighting scenario */
-	bool IsNewLightingScenarioState(bool bExistingState) const;
+	UE_API bool IsNewLightingScenarioState(bool bExistingState) const;
 
-	void SetIsLightingScenario(bool bNewLightingScenario);
+	UE_API void SetIsLightingScenario(bool bNewLightingScenario);
 
 	/** @return whether any of the currently selected levels is dirty */
-	bool AreAnySelectedLevelsDirty() const;
+	UE_API bool AreAnySelectedLevelsDirty() const;
 
 	/** @return	whether at least one actor is selected */
-	bool AreActorsSelected() const;
+	UE_API bool AreActorsSelected() const;
 
 	/** @return whether any of the currently selected levels can be converted to the specified actor bExternal packaging */
-	bool CanConvertAnyLevelToExternalActors(bool bExternal) const;
+	UE_API bool CanConvertAnyLevelToExternalActors(bool bExternal) const;
 
 	/** @return whether moving the selected actors to the selected level is a valid action */
-	bool IsValidMoveActorsToLevel() const;
+	UE_API bool IsValidMoveActorsToLevel() const;
 
 	/** @return whether moving the selected foliage to the selected level is a valid action */
-	bool IsValidMoveFoliageToLevel() const;
+	UE_API bool IsValidMoveFoliageToLevel() const;
 
 	/** delegate used to pickup when the selection has changed */
-	void OnActorSelectionChanged(UObject* obj);
+	UE_API void OnActorSelectionChanged(UObject* obj);
 
 	/** Sets a flag to re-cache whether the selected actors move to the selected level is valid */
-	void OnActorOrLevelSelectionChanged();
+	UE_API void OnActorOrLevelSelectionChanged();
 
 	/** @return	whether 'display paths' is enabled */
-	bool GetDisplayPathsState() const;
+	UE_API bool GetDisplayPathsState() const;
 
 	/** Sets 'display paths', whether to show long package name in level display name */
-	void SetDisplayPathsState(bool bDisplayPaths);
+	UE_API void SetDisplayPathsState(bool bDisplayPaths);
 
 	/** @return	whether 'display actors count' is enabled */
-	bool GetDisplayActorsCountState() const;
+	UE_API bool GetDisplayActorsCountState() const;
 
 	/** Sets 'display actors count', whether to show actors count next to level name */
-	void SetDisplayActorsCountState(bool bDisplayActorsCount);
+	UE_API void SetDisplayActorsCountState(bool bDisplayActorsCount);
 
 	/**	Broadcasts whenever items selection has changed */
 	FSimpleEvent SelectionChanged;
-	void BroadcastSelectionChanged();
+	UE_API void BroadcastSelectionChanged();
 
 	/**	Broadcasts whenever items collection has changed */
 	FSimpleEvent CollectionChanged;
-	void BroadcastCollectionChanged();
+	UE_API void BroadcastCollectionChanged();
 		
 	/** Broadcasts whenever items hierarchy has changed */
 	FSimpleEvent HierarchyChanged;
-	void BroadcastHierarchyChanged();
+	UE_API void BroadcastHierarchyChanged();
 
 	/** Broadcasts before levels are unloaded */
 	FSimpleEvent PreLevelsUnloaded;
-	void BroadcastPreLevelsUnloaded();
+	UE_API void BroadcastPreLevelsUnloaded();
 
 	/** Broadcasts after levels are unloaded */
 	FSimpleEvent PostLevelsUnloaded;
-	void BroadcastPostLevelsUnloaded();
+	UE_API void BroadcastPostLevelsUnloaded();
 	
 	/** Editable world axis length  */
-	static double EditableAxisLength();
+	static UE_API double EditableAxisLength();
 
 	/** Editable world bounds */
-	static FBox EditableWorldArea();
+	static UE_API FBox EditableWorldArea();
 
 	/**  */
-	static void SCCCheckOut(const FLevelModelList& InList);
-	static void SCCCheckIn(const FLevelModelList& InList);
-	static void SCCOpenForAdd(const FLevelModelList& InList);
-	static void SCCHistory(const FLevelModelList& InList);
-	static void SCCRefresh(const FLevelModelList& InList);
-	static void SCCDiffAgainstDepot(const FLevelModelList& InList, UEditorEngine* InEditor);
+	static UE_API void SCCCheckOut(const FLevelModelList& InList);
+	static UE_API void SCCCheckIn(const FLevelModelList& InList);
+	static UE_API void SCCOpenForAdd(const FLevelModelList& InList);
+	static UE_API void SCCHistory(const FLevelModelList& InList);
+	static UE_API void SCCRefresh(const FLevelModelList& InList);
+	static UE_API void SCCDiffAgainstDepot(const FLevelModelList& InList, UEditorEngine* InEditor);
 	
 	/** @return	List of valid level package names from a specified level model list*/
-	static TArray<FName> GetPackageNamesList(const FLevelModelList& InList);
+	static UE_API TArray<FName> GetPackageNamesList(const FLevelModelList& InList);
 	
 	/** @return	List of valid level package filenames from a specified level model list*/
-	static TArray<FString> GetFilenamesList(const FLevelModelList& InList);
+	static UE_API TArray<FString> GetFilenamesList(const FLevelModelList& InList);
 	
 	/** @return	List of valid packages from a specified level model list*/
-	static TArray<UPackage*> GetPackagesList(const FLevelModelList& InList);
+	static UE_API TArray<UPackage*> GetPackagesList(const FLevelModelList& InList);
 	
 	/** @return	List of valid level objects from a specified level model list*/
-	static TArray<ULevel*> GetLevelObjectList(const FLevelModelList& InList);
+	static UE_API TArray<ULevel*> GetLevelObjectList(const FLevelModelList& InList);
 
 	/** @return	List of loaded level models from a specified level model list*/
-	static FLevelModelList GetLoadedLevels(const FLevelModelList& InList);
+	static UE_API FLevelModelList GetLoadedLevels(const FLevelModelList& InList);
 
 	/** @return	List of all level models found while traversing hierarchy of specified level models */
-	static FLevelModelList GetLevelsHierarchy(const FLevelModelList& InList);
+	static UE_API FLevelModelList GetLevelsHierarchy(const FLevelModelList& InList);
 
 	/** @return	Total bounding box of specified level models */
-	static FBox GetLevelsBoundingBox(const FLevelModelList& InList, bool bIncludeChildren);
+	static UE_API FBox GetLevelsBoundingBox(const FLevelModelList& InList, bool bIncludeChildren);
 
 	/** @return	Total bounding box of specified visible level models */
-	static FBox GetVisibleLevelsBoundingBox(const FLevelModelList& InList, bool bIncludeChildren);
+	static UE_API FBox GetVisibleLevelsBoundingBox(const FLevelModelList& InList, bool bIncludeChildren);
 
 	/** @return	The UICommandList supported by this collection */
-	const TSharedRef<FUICommandList> GetCommandList() const;
+	UE_API const TSharedRef<FUICommandList> GetCommandList() const;
 
 	/**  */
-	void LoadSettings();
+	UE_API void LoadSettings();
 	
 	/**  */
-	void SaveSettings();
+	UE_API void SaveSettings();
 
 protected:
 	/** Refreshes current cached data */
-	void RefreshBrowser_Executed();
+	UE_API void RefreshBrowser_Executed();
 	
 	/** Load selected levels to the world */
-	void LoadSelectedLevels_Executed();
+	UE_API void LoadSelectedLevels_Executed();
 
 	/** Unload selected level from the world */
-	void UnloadSelectedLevels_Executed();
+	UE_API void UnloadSelectedLevels_Executed();
 
 	/** Make this Level the Current Level */
-	void MakeLevelCurrent_Executed();
+	UE_API void MakeLevelCurrent_Executed();
 
 	/** Find selected levels in Content Browser */
-	void FindInContentBrowser_Executed();
+	UE_API void FindInContentBrowser_Executed();
 
 	/** Is FindInContentBrowser a valid action */
-	bool IsValidFindInContentBrowser();
+	UE_API bool IsValidFindInContentBrowser();
 
 	/** Moves the selected actors to this level */
-	void MoveActorsToSelected_Executed();
+	UE_API void MoveActorsToSelected_Executed();
 
 	/** Moves the selected foliage to this level */
-	void MoveFoliageToSelected_Executed();
+	UE_API void MoveFoliageToSelected_Executed();
 
 	/** Saves selected levels */
-	void SaveSelectedLevels_Executed();
+	UE_API void SaveSelectedLevels_Executed();
 
 	/** Saves selected level under new name */
-	void SaveSelectedLevelAs_Executed();
+	UE_API void SaveSelectedLevelAs_Executed();
 	
 	/** Migrate selected levels */
-	void MigrateSelectedLevels_Executed();
+	UE_API void MigrateSelectedLevels_Executed();
 
 	/** Expand selected items hierarchy */
-	void ExpandSelectedItems_Executed();
+	UE_API void ExpandSelectedItems_Executed();
 			
 	/** Check-Out selected levels from SCC */
-	void OnSCCCheckOut();
+	UE_API void OnSCCCheckOut();
 
 	/** Mark for Add selected levels from SCC */
-	void OnSCCOpenForAdd();
+	UE_API void OnSCCOpenForAdd();
 
 	/** Check-In selected levels from SCC */
-	void OnSCCCheckIn();
+	UE_API void OnSCCCheckIn();
 
 	/** Shows the SCC History of selected levels */
-	void OnSCCHistory();
+	UE_API void OnSCCHistory();
 
 	/** Refreshes the states selected levels from SCC */
-	void OnSCCRefresh();
+	UE_API void OnSCCRefresh();
 
 	/** Diffs selected levels from with those in the SCC depot */
-	void OnSCCDiffAgainstDepot();
+	UE_API void OnSCCDiffAgainstDepot();
 
 	/** Enable source control features */
-	void OnSCCConnect() const;
+	UE_API void OnSCCConnect() const;
 
 	/** Selects all levels in the collection view model */
-	void SelectAllLevels_Executed();
+	UE_API void SelectAllLevels_Executed();
 
 	/** De-selects all levels in the collection view model */
-	void DeselectAllLevels_Executed();
+	UE_API void DeselectAllLevels_Executed();
 
 	/** Inverts level selection in the collection view model */
-	void InvertSelection_Executed();
+	UE_API void InvertSelection_Executed();
 
 	/** Adds the Actors in the selected Levels from the viewport's existing selection */
-	void SelectActors_Executed();
+	UE_API void SelectActors_Executed();
 
 	/** Removes the Actors in the selected Levels from the viewport's existing selection */
-	void DeselectActors_Executed();
+	UE_API void DeselectActors_Executed();
 
 	/** Set level `Use External Actors` to bExternal  */
-	void ConvertLevelToExternalActors_Executed(bool bExternal);
+	UE_API void ConvertLevelToExternalActors_Executed(bool bExternal);
 
 	/** Toggles selected levels to a visible state in the viewports for editor worlds */
-	void ShowInEditorSelectedLevels_Executed();
+	UE_API void ShowInEditorSelectedLevels_Executed();
 
 	/** Toggles selected levels to an invisible state in the viewports for editor worlds */
-	void HideInEditorSelectedLevels_Executed();
+	UE_API void HideInEditorSelectedLevels_Executed();
 
 	/** Toggles the selected levels to a visible state for editor worlds; toggles all other levels to an invisible state. */
-	void ShowInEditorOnlySelectedLevels_Executed();
+	UE_API void ShowInEditorOnlySelectedLevels_Executed();
 
 	/** Toggles the selected levels to an invisible state for editor worlds; toggles all other levels to a visible state. */
-	void ShowInEditorAllButSelectedLevels_Executed();
+	UE_API void ShowInEditorAllButSelectedLevels_Executed();
 
 	/** Toggles all levels to a visible state in the viewports for editor worlds */
-	void ShowInEditorAllLevels_Executed();
+	UE_API void ShowInEditorAllLevels_Executed();
 
 	/** Hides all levels to an invisible state in the viewports for editor worlds */
-	void HideInEditorAllLevels_Executed();
+	UE_API void HideInEditorAllLevels_Executed();
 
 	/** Toggles selected levels to a visible state in the viewports for game worlds */
-	void ShowInGameSelectedLevels_Executed();
+	UE_API void ShowInGameSelectedLevels_Executed();
 
 	/** Toggles selected levels to an invisible state in the viewports */
-	void HideInGameSelectedLevels_Executed();
+	UE_API void HideInGameSelectedLevels_Executed();
 
 	/** Toggles the selected levels to a visible state for game worlds; toggles all other levels to an invisible state. */
-	void ShowInGameOnlySelectedLevels_Executed();
+	UE_API void ShowInGameOnlySelectedLevels_Executed();
 
 	/** Toggles the selected levels to an invisible state for game worlds; toggles all other levels to a visible state. */
-	void ShowInGameAllButSelectedLevels_Executed();
+	UE_API void ShowInGameAllButSelectedLevels_Executed();
 
 	/** Toggles all levels to a visible state in the viewports for game worlds */
-	void ShowInGameAllLevels_Executed();
+	UE_API void ShowInGameAllLevels_Executed();
 
 	/** Hides all levels to an invisible state in the viewports for game worlds */
-	void HideInGameAllLevels_Executed();
+	UE_API void HideInGameAllLevels_Executed();
 	
 	/** Locks selected levels */
-	void LockSelectedLevels_Executed();
+	UE_API void LockSelectedLevels_Executed();
 
 	/** Unlocks selected levels */
-	void UnlockSelectedLevels_Executed();
+	UE_API void UnlockSelectedLevels_Executed();
 
 	/** Toggles the selected levels to a locked state; toggles all other levels to an unlocked state */
-	void LockOnlySelectedLevels_Executed();
+	UE_API void LockOnlySelectedLevels_Executed();
 
 	/** Toggles the selected levels to an unlocked state; toggles all other levels to a locked state */
-	void LockAllButSelectedLevels_Executed();
+	UE_API void LockAllButSelectedLevels_Executed();
 
 	/** Locks all levels */
-	void LockAllLevels_Executed();
+	UE_API void LockAllLevels_Executed();
 
 	/** Unlocks all levels */
-	void UnlockAllLevels_Executed();
+	UE_API void UnlockAllLevels_Executed();
 	
 	/** Toggle all read-only levels */
-	void ToggleReadOnlyLevels_Executed();
+	UE_API void ToggleReadOnlyLevels_Executed();
 
 	/** true if the SCC Check-Out option is available */
 	bool CanExecuteSCCCheckOut() const
@@ -524,41 +526,41 @@ protected:
 	}
 	
 	/** Fills MenuBulder with Lock level related commands */
-	void FillLockSubMenu(FMenuBuilder& MenuBuilder);
+	UE_API void FillLockSubMenu(FMenuBuilder& MenuBuilder);
 
 	/** Fills MenuBulder with level visisbility related commands */
-	void FillEditorVisibilitySubMenu(FMenuBuilder& MenuBuilder);
+	UE_API void FillEditorVisibilitySubMenu(FMenuBuilder& MenuBuilder);
 
 	/** Fills MenuBulder with level visisbility related commands */
-	void FillGameVisibilitySubMenu(FMenuBuilder& MenuBuilder);
+	UE_API void FillGameVisibilitySubMenu(FMenuBuilder& MenuBuilder);
 
 	/** Fills MenuBulder with SCC related commands */
-	void FillSourceControlSubMenu(FMenuBuilder& MenuBuilder);
+	UE_API void FillSourceControlSubMenu(FMenuBuilder& MenuBuilder);
 				
 protected:
 	/**  */
-	virtual void Initialize(UWorld* InWorld);
+	UE_API virtual void Initialize(UWorld* InWorld);
 	
 	/**  */
-	virtual void BindCommands();
+	UE_API virtual void BindCommands();
 
 	/** Removes the Actors in all read-only Levels from the viewport's existing selection */
-	void DeselectActorsInAllReadOnlyLevel(const FLevelModelList& InLevelList);
+	UE_API void DeselectActorsInAllReadOnlyLevel(const FLevelModelList& InLevelList);
 
 	/** Removes the Actors in all read-only Levels from the viewport's existing selection */
-	void DeselectSurfaceInAllReadOnlyLevel(const FLevelModelList& InLevelList);
+	UE_API void DeselectSurfaceInAllReadOnlyLevel(const FLevelModelList& InLevelList);
 	
 	/** Called whenever level collection has been changed */
-	virtual void OnLevelsCollectionChanged();
+	UE_API virtual void OnLevelsCollectionChanged();
 	
 	/** Called whenever level selection has been changed */
-	virtual void OnLevelsSelectionChanged();
+	UE_API virtual void OnLevelsSelectionChanged();
 
 	/** Called whenever level selection has been changed outside of this module, usually via World->SetSelectedLevels */
-	void OnLevelsSelectionChangedOutside();
+	UE_API void OnLevelsSelectionChangedOutside();
 	
 	/** Called whenever level collection hierarchy has been changed */
-	virtual void OnLevelsHierarchyChanged();
+	UE_API virtual void OnLevelsHierarchyChanged();
 
 	/** Called before loading specified level models into editor */
 	virtual void OnPreLoadLevels(const FLevelModelList& InList) {};
@@ -567,25 +569,25 @@ protected:
 	virtual void OnPreShowLevels(const FLevelModelList& InList) {};
 
 	/** Called when level was added to the world */
-	void OnLevelAddedToWorld(ULevel* InLevel, UWorld* InWorld);
+	UE_API void OnLevelAddedToWorld(ULevel* InLevel, UWorld* InWorld);
 
 	/** Called when level was removed from the world */
-	void OnLevelRemovedFromWorld(ULevel* InLevel, UWorld* InWorld);
+	UE_API void OnLevelRemovedFromWorld(ULevel* InLevel, UWorld* InWorld);
 
 	/** Handler for FEditorSupportDelegates::RedrawAllViewports event */
-	void OnRedrawAllViewports();
+	UE_API void OnRedrawAllViewports();
 
 	/** Handler for when an actor was added to a level */
-	void OnLevelActorAdded(AActor* InActor);	
+	UE_API void OnLevelActorAdded(AActor* InActor);	
 
 	/** Handler for when an actor was removed from a level */
-	void OnLevelActorDeleted(AActor* InActor);
+	UE_API void OnLevelActorDeleted(AActor* InActor);
 	
 	/** Handler for level filter collection changes */
-	void OnFilterChanged();
+	UE_API void OnFilterChanged();
 
 	/** Caches the variables for which SCC menu options are available */
-	void CacheCanExecuteSourceControlVars() const;
+	UE_API void CacheCanExecuteSourceControlVars() const;
 		
 protected:
 	
@@ -721,3 +723,5 @@ struct FTiledLandscapeImportSettings
 
 	TArray<LandscapeLayerSettings>	LandscapeLayerSettingsList;
 };
+
+#undef UE_API
diff --git a/Engine/Source/Editor/WorldBrowser/Private/LevelFolders.h b/Engine/Source/Editor/WorldBrowser/Private/LevelFolders.h
--- a/Engine/Source/Editor/WorldBrowser/Private/LevelFolders.h
+++ b/Engine/Source/Editor/WorldBrowser/Private/LevelFolders.h
@@ -12,6 +12,8 @@
 
 #include "LevelFolders.generated.h"
 
+#define UE_API WORLDBROWSER_API
+
 typedef		FName		FLevelModelKey;
 
 /** Broadcasted when an editor-only folder has been created for a level */
@@ -56,15 +58,15 @@ public:
 
 
 /** The class for managing in-memory representations of level folders in the editor */
-struct WORLDBROWSER_API FLevelFolders : public FGCObject
+struct FLevelFolders : public FGCObject
 {
-	FLevelFolders();
-	~FLevelFolders();
+	UE_API FLevelFolders();
+	UE_API ~FLevelFolders();
 
 public:
 
 	//~ FGCObject
-	virtual void AddReferencedObjects(FReferenceCollector& Collector) override;
+	UE_API virtual void AddReferencedObjects(FReferenceCollector& Collector) override;
 	virtual FString GetReferencerName() const
 	{
 		return "FLevelFolders";
@@ -76,72 +78,72 @@ public:
 	static bool IsAvailable() { return Singleton != nullptr; }
 
 	/** Grants access to the singleton object if it's available */
-	static FLevelFolders& Get();
+	static UE_API FLevelFolders& Get();
 
 	/** Initialize the singleton */
-	static void Init();
+	static UE_API void Init();
 
 	/** Clean up the singleton */
-	static void Cleanup();
+	static UE_API void Cleanup();
 
 
 	/** Folder events */
-	static FOnLevelFolderCreate OnFolderCreate;
-	static FOnLevelFolderDelete OnFolderDelete;
-	static FOnLevelFolderMove OnFolderMove;
+	static UE_API FOnLevelFolderCreate OnFolderCreate;
+	static UE_API FOnLevelFolderDelete OnFolderDelete;
+	static UE_API FOnLevelFolderMove OnFolderMove;
 
 	/** Gets all folder properties for a specified level */
-	TMap<FName, FLevelFolderProps>& GetFolderProperties(TSharedRef<FLevelModel> LevelModel);
+	UE_API TMap<FName, FLevelFolderProps>& GetFolderProperties(TSharedRef<FLevelModel> LevelModel);
 
 	/** Gets the folder properties for a specified path within the level */
-	FLevelFolderProps* GetFolderProperties(TSharedRef<FLevelModel> LevelModel, FName InPath);
+	UE_API FLevelFolderProps* GetFolderProperties(TSharedRef<FLevelModel> LevelModel, FName InPath);
 
 	/** Gets the default folder name for the given path. */
-	FName GetDefaultFolderName(TSharedRef<FLevelModel> LevelModel, FName ParentPath);
+	UE_API FName GetDefaultFolderName(TSharedRef<FLevelModel> LevelModel, FName ParentPath);
 
 	/** Creates a new folder with the given name for the current level selection */
-	void CreateFolderContainingSelectedLevels(TSharedRef<FLevelCollectionModel> WorldModel, TSharedRef<FLevelModel> LevelModel, FName InPath);
+	UE_API void CreateFolderContainingSelectedLevels(TSharedRef<FLevelCollectionModel> WorldModel, TSharedRef<FLevelModel> LevelModel, FName InPath);
 
 	/** Creates a folder for the level model with the given path name */
-	void CreateFolder(TSharedRef<FLevelModel> LevelModel, FName InPath);
+	UE_API void CreateFolder(TSharedRef<FLevelModel> LevelModel, FName InPath);
 
 	/** Renames a folder. The folder with the old name is removed from the folder props */
-	bool RenameFolder(TSharedRef<FLevelModel> LevelModel, FName OldPath, FName NewPath);
+	UE_API bool RenameFolder(TSharedRef<FLevelModel> LevelModel, FName OldPath, FName NewPath);
 
 	/** Deletes a folder and all saved properties */
-	void DeleteFolder(TSharedRef<FLevelModel> LevelModel, FName FolderToDelete);
+	UE_API void DeleteFolder(TSharedRef<FLevelModel> LevelModel, FName FolderToDelete);
 
 	/** Rebuilds the folder list for the level */
-	void RebuildFolderList(TSharedRef<FLevelModel> LevelModel);
+	UE_API void RebuildFolderList(TSharedRef<FLevelModel> LevelModel);
 
 	/** Saves the level model when the world is saved */
-	void SaveLevel(TSharedRef<FLevelModel> LevelModel);
+	UE_API void SaveLevel(TSharedRef<FLevelModel> LevelModel);
 
 private:
 
 	/** Sets the folder path for the current level selection */
-	void SetSelectedLevelFolderPath(TSharedRef<FLevelCollectionModel> WorldModel, TSharedRef<FLevelModel> LevelModel, FName InPath) const;
+	UE_API void SetSelectedLevelFolderPath(TSharedRef<FLevelCollectionModel> WorldModel, TSharedRef<FLevelModel> LevelModel, FName InPath) const;
 
 	/** Remove references to folder arrays for unloaded levels */
-	void Housekeeping();
+	UE_API void Housekeeping();
 
 	/** Create a folder container for the specified level model */
-	UEditorLevelFolders& InitializeForLevel(TSharedRef<FLevelModel> LevelModel);
+	UE_API UEditorLevelFolders& InitializeForLevel(TSharedRef<FLevelModel> LevelModel);
 
 	/** Checks if the supplied path is a descendant of the parent path */
-	bool PathIsChildOf(const FString& InPotentialChild, const FString& InParent);
+	UE_API bool PathIsChildOf(const FString& InPotentialChild, const FString& InParent);
 
 	/** Gets the folder information for the given level, or creates it if it's not in memory */
-	UEditorLevelFolders& GetOrCreateFoldersForLevel(TSharedRef<FLevelModel> LevelModel);
+	UE_API UEditorLevelFolders& GetOrCreateFoldersForLevel(TSharedRef<FLevelModel> LevelModel);
 
 	/** Creates new folder information for the level */
-	UEditorLevelFolders& Initialize(TSharedRef<FLevelModel> LevelModel);
+	UE_API UEditorLevelFolders& Initialize(TSharedRef<FLevelModel> LevelModel);
 
 	/** Adds a folder for the level without triggering any events */
-	bool AddFolder(TSharedRef<FLevelModel> LevelModel, FName InPath);
+	UE_API bool AddFolder(TSharedRef<FLevelModel> LevelModel, FName InPath);
 
 	/** Gets the selected levels in the world model */
-	FLevelModelList GetSelectedLevels(TSharedRef<FLevelCollectionModel> WorldModel, TSharedRef<FLevelModel> LevelModel) const;
+	UE_API FLevelModelList GetSelectedLevels(TSharedRef<FLevelCollectionModel> WorldModel, TSharedRef<FLevelModel> LevelModel) const;
 
 private:
 
@@ -152,5 +154,7 @@ private:
 	TMap<FLevelModelKey, TWeakPtr<FLevelModel>> TemporaryModelObjects;
 
 	/** Singleton object to be maintained by the world browser module */
-	static FLevelFolders* Singleton;
-};
\ No newline at end of file
+	static UE_API FLevelFolders* Singleton;
+};
+
+#undef UE_API
diff --git a/Engine/Source/Editor/WorldBrowser/Private/LevelModel.h b/Engine/Source/Editor/WorldBrowser/Private/LevelModel.h
--- a/Engine/Source/Editor/WorldBrowser/Private/LevelModel.h
+++ b/Engine/Source/Editor/WorldBrowser/Private/LevelModel.h
@@ -3,6 +3,8 @@
 
 #include "CoreMinimal.h"
 
+#define UE_API WORLDBROWSER_API
+
 class AActor;
 struct FAssetData;
 class FLevelCollectionModel;
@@ -27,7 +29,7 @@ struct FLevelModelVisitor
 /**
  * Interface for non-UI presentation logic for a level in a world
  */
-class WORLDBROWSER_API FLevelModel
+class FLevelModel
 	: public TSharedFromThis<FLevelModel>	
 {
 public:
@@ -36,36 +38,36 @@ public:
 	DECLARE_EVENT( FLevelModel, FSimpleEvent );
 	
 public:
-	FLevelModel(FLevelCollectionModel& InLevelCollectionModel);
+	UE_API FLevelModel(FLevelCollectionModel& InLevelCollectionModel);
 
-	virtual ~FLevelModel();
+	UE_API virtual ~FLevelModel();
 	
 	/** Traverses level model hierarchy */
-	void Accept(FLevelModelVisitor& Vistor);
+	UE_API void Accept(FLevelModelVisitor& Vistor);
 
 	/** Sets level selection flag */
-	void SetLevelSelectionFlag(bool bExpanded);
+	UE_API void SetLevelSelectionFlag(bool bExpanded);
 	
 	/** @return Level selection flag */
-	bool GetLevelSelectionFlag() const;
+	UE_API bool GetLevelSelectionFlag() const;
 		
 	/** Sets level child hierarchy expansion flag */
-	void SetLevelExpansionFlag(bool bExpanded);
+	UE_API void SetLevelExpansionFlag(bool bExpanded);
 	
 	/** @return Level child hierarchy expansion flag */
-	bool GetLevelExpansionFlag() const;
+	UE_API bool GetLevelExpansionFlag() const;
 
 	/** Sets level filtered out flag */
-	void SetLevelFilteredOutFlag(bool bFiltredOut);
+	UE_API void SetLevelFilteredOutFlag(bool bFiltredOut);
 	
 	/** @return Whether this level model was filtered out */
-	bool GetLevelFilteredOutFlag() const;
+	UE_API bool GetLevelFilteredOutFlag() const;
 
 	/**	@return	Level display name */
-	FString GetDisplayName() const;
+	UE_API FString GetDisplayName() const;
 	
 	/**	@return	Level package file name */
-	FString GetPackageFileName() const;
+	UE_API FString GetPackageFileName() const;
 
 	/**	@return	Whether level model has valid package file */
 	virtual bool HasValidPackage() const { return true; };
@@ -86,137 +88,137 @@ public:
 	virtual void UpdateAsset(const FAssetData& AssetData) = 0;
 	
 	/** Refreshes cached data */
-	virtual void Update();
+	UE_API virtual void Update();
 	
 	/** Refreshes visual information */
-	virtual void UpdateVisuals();
+	UE_API virtual void UpdateVisuals();
 
 	/**	@return	Whether level is in PIE/SIE mode */
-	bool IsSimulating() const;
+	UE_API bool IsSimulating() const;
 	
 	/** @return Whether level is CurrentLevel */
-	bool IsCurrent() const;
+	UE_API bool IsCurrent() const;
 
 	/** @return Whether level is PersistentLevel */
-	bool IsPersistent() const;
+	UE_API bool IsPersistent() const;
 
 	/** @return Whether level is editable */
-	bool IsEditable() const;
+	UE_API bool IsEditable() const;
 
 	/** @return Whether level is dirty */
-	bool IsDirty() const;
+	UE_API bool IsDirty() const;
 
 	/** @return Whether level is a lighting scenario */
-	bool IsLightingScenario() const;
+	UE_API bool IsLightingScenario() const;
 
-	void SetIsLightingScenario(bool bNew);
+	UE_API void SetIsLightingScenario(bool bNew);
 
 	/** @return Whether level has loaded content */
-	bool IsLoaded() const;
+	UE_API bool IsLoaded() const;
 
 	/** @return Whether level is in process of loading content */
-	bool IsLoading() const;
+	UE_API bool IsLoading() const;
 
 	/**	@return Whether level is visible in editor worlds */
-	bool IsVisibleInEditor() const;
+	UE_API bool IsVisibleInEditor() const;
 
 	/**	@return Whether level is visible in game worlds */
-	bool IsVisibleInGame() const;
+	UE_API bool IsVisibleInGame() const;
 
 	/**	@return Whether level is locked */
-	bool IsLocked() const;
+	UE_API bool IsLocked() const;
 
 	/** @return Whether level package file is read only */
-	bool IsFileReadOnly() const;
+	UE_API bool IsFileReadOnly() const;
 
 	/** @return Whether this level can be add/removed/unloaded by user */
-	bool IsUserManaged() const;
+	UE_API bool IsUserManaged() const;
 
 	/** @return Whether the streaming level object is transient */
 	virtual bool IsTransient() const { return false; }
 
 	/** Loads level into editor */
-	virtual void LoadLevel();
+	UE_API virtual void LoadLevel();
 
 	/** Sets the Level's visibility */
-	virtual void SetVisibleInEditor(bool bVisible);
+	UE_API virtual void SetVisibleInEditor(bool bVisible);
 
 	/** Static function analog to the non-static SetVisible but applied to an TArray of FLevelModel elements */
-	static void SetVisibleInEditor(TArray<FLevelModel*>& LevelModels, const TArray<bool>& bAreVisible);
+	static UE_API void SetVisibleInEditor(TArray<FLevelModel*>& LevelModels, const TArray<bool>& bAreVisible);
 
-	static void SetVisibleInGame(TArray<FLevelModel*>& LevelModels, const TArray<bool>& bAreVisible);
+	static UE_API void SetVisibleInGame(TArray<FLevelModel*>& LevelModels, const TArray<bool>& bAreVisible);
 
 	/** Sets the Level's locked/unlocked state */
-	void SetLocked(bool bLocked);
+	UE_API void SetLocked(bool bLocked);
 
 	/** Sets Level as current in the world */
-	void MakeLevelCurrent();
+	UE_API void MakeLevelCurrent();
 	
 	/** @return Whether specified point is hovering level */
-	virtual bool HitTest2D(const FVector2D& Point) const;
+	UE_API virtual bool HitTest2D(const FVector2D& Point) const;
 	
 	/** @return Level top left corner position */
-	virtual FVector2D GetLevelPosition2D() const;
+	UE_API virtual FVector2D GetLevelPosition2D() const;
 
 	/** @return XY size of level */
-	virtual FVector2D GetLevelSize2D() const;
+	UE_API virtual FVector2D GetLevelSize2D() const;
 
 	/** @return Level bounding box */
-	virtual FBox GetLevelBounds() const;
+	UE_API virtual FBox GetLevelBounds() const;
 	
 	/** @return level translation delta, when user moving level item */
-	FVector2D GetLevelTranslationDelta() const;
+	UE_API FVector2D GetLevelTranslationDelta() const;
 
 	/** Sets new translation delta to this model and all descendants*/
-	void SetLevelTranslationDelta(FVector2D InAbsoluteDelta);
+	UE_API void SetLevelTranslationDelta(FVector2D InAbsoluteDelta);
 
 	/** @return level color, used for visualization. (Show -> Advanced -> Level Coloration) */
-	virtual FLinearColor GetLevelColor() const;
+	UE_API virtual FLinearColor GetLevelColor() const;
 
 	/** Sets level color, used for visualization. (Show -> Advanced -> Level Coloration) */
-	virtual void SetLevelColor(FLinearColor InColor);
+	UE_API virtual void SetLevelColor(FLinearColor InColor);
 
 	/** Whether level should be drawn in world composition view */
-	virtual bool IsVisibleInCompositionView() const;
+	UE_API virtual bool IsVisibleInCompositionView() const;
 	
 	/**	@return Whether level has associated blueprint script */
-	bool HasKismet() const;
+	UE_API bool HasKismet() const;
 
 	/**	Opens level associated blueprint script */
-	void OpenKismet();
+	UE_API void OpenKismet();
 
 	/** 
 	 * Sets parent for this item
 	 * @return false in case attaching has failed
 	 */
-	bool AttachTo(TSharedPtr<FLevelModel> InParent);
+	UE_API bool AttachTo(TSharedPtr<FLevelModel> InParent);
 
 	/**	Notifies level model that filters has been changed */
-	void OnFilterChanged();
+	UE_API void OnFilterChanged();
 
 	/**	@return Level child hierarchy */
-	const FLevelModelList& GetChildren() const;
+	UE_API const FLevelModelList& GetChildren() const;
 	
 	/**	@return Parent level model */
-	TSharedPtr<FLevelModel> GetParent() const;
+	UE_API TSharedPtr<FLevelModel> GetParent() const;
 	
 	/**	Sets link to a parent model  */
-	void SetParent(TSharedPtr<FLevelModel>);
+	UE_API void SetParent(TSharedPtr<FLevelModel>);
 
 	/**	Removes all entries from children list*/
-	void RemoveAllChildren();
+	UE_API void RemoveAllChildren();
 
 	/**	Removes specific child */
-	void RemoveChild(TSharedPtr<FLevelModel> InChild);
+	UE_API void RemoveChild(TSharedPtr<FLevelModel> InChild);
 	
 	/**	Adds new entry to a children list */
-	void AddChild(TSharedPtr<FLevelModel> InChild);
+	UE_API void AddChild(TSharedPtr<FLevelModel> InChild);
 		
 	/**	@return Whether this model has in ancestors specified level model */
-	bool HasAncestor(const TSharedPtr<FLevelModel>& InLevel) const;
+	UE_API bool HasAncestor(const TSharedPtr<FLevelModel>& InLevel) const;
 
 	/**	@return Whether this model has in descendants specified level model */
-	bool HasDescendant(const TSharedPtr<FLevelModel>& InLevel) const;
+	UE_API bool HasDescendant(const TSharedPtr<FLevelModel>& InLevel) const;
 	
 	/** Returns the folder path that the level should use when displayed in the world hierarchy */
 	virtual FName GetFolderPath() const { return NAME_None; }
@@ -228,22 +230,22 @@ public:
 	virtual bool HasFolderSupport() const { return false; }
 
 	/**	@return Handles drop operation */
-	virtual void OnDrop(const TSharedPtr<FLevelDragDropOp>& Op);
+	UE_API virtual void OnDrop(const TSharedPtr<FLevelDragDropOp>& Op);
 	
 	/**	@return Whether it's possible to drop onto this level */
-	virtual bool IsGoodToDrop(const TSharedPtr<FLevelDragDropOp>& Op) const;
+	UE_API virtual bool IsGoodToDrop(const TSharedPtr<FLevelDragDropOp>& Op) const;
 
 	/** Notification when level was added(shown) to world */
-	virtual void OnLevelAddedToWorld(ULevel* InLevel);
+	UE_API virtual void OnLevelAddedToWorld(ULevel* InLevel);
 
 	/** Notification when level was removed(hidden) from world */
-	virtual void OnLevelRemovedFromWorld();
+	UE_API virtual void OnLevelRemovedFromWorld();
 
 	/** Notification on level reparenting */
 	virtual void OnParentChanged() {};
 
 	/** Event when level model has been changed */
-	void BroadcastLevelChanged();
+	UE_API void BroadcastLevelChanged();
 	
 	/*
 	 *
@@ -256,17 +258,17 @@ public:
 	};
 	
 	/** Updates this level simulation status  */
-	void UpdateSimulationStatus(ULevelStreaming* StreamingLevel);
+	UE_API void UpdateSimulationStatus(ULevelStreaming* StreamingLevel);
 
 	/**	Broadcasts whenever level has changed */
 	FSimpleEvent ChangedEvent;
-	void BroadcastChangedEvent();
+	UE_API void BroadcastChangedEvent();
 
 	/** Deselects all Actors in this level */
-	void DeselectAllActors();
+	UE_API void DeselectAllActors();
 
 	/** Deselects all BSP surfaces in this level */
-	void DeselectAllSurfaces();
+	UE_API void DeselectAllSurfaces();
 
 	/**
 	 *	Selects in the Editor all the Actors assigned to the Level, based on the specified conditions.
@@ -276,40 +278,40 @@ public:
 	 *	@param	bSelectEvenIfHidden		if true actors that are hidden will be selected; If false, they will be skipped
 	 *	@param	Filter					Only actors which pass the filters restrictions will be selected
 	 */
-	void SelectActors(bool bSelect, bool bNotify, bool bSelectEvenIfHidden, 
+	UE_API void SelectActors(bool bSelect, bool bNotify, bool bSelectEvenIfHidden, 
 						const TSharedPtr<ActorFilter>& Filter = TSharedPtr<ActorFilter>(NULL));
 	
 	/**
 	 * Set the level to use external actors and convert its actors to external packaging or vice versa
 	 * @param bUseExternal Convert to external if `bUseExternal` is true, convert to internal otherwise
 	 */
-	void ConvertLevelToExternalActors(bool bUseExternal);
+	UE_API void ConvertLevelToExternalActors(bool bUseExternal);
 
 	/**
 	 * Return if the level can be converted to the specified external packaging mode.
 	 * @param bToExternal The packaging mode we want to convert to. (i.e true: external, false: internal)
 	 * @return true if we can convert the level to the `bToExternal` packaging mode.
 	 */
-	bool CanConvertLevelToExternalActors(bool bToExternal);
+	UE_API bool CanConvertLevelToExternalActors(bool bToExternal);
 
 	/** Updates cached value of level actors count */
-	void UpdateLevelActorsCount();
+	UE_API void UpdateLevelActorsCount();
 
 	/** Updates cached value of level display name */
-	void UpdateDisplayName();
+	UE_API void UpdateDisplayName();
 
 	/** @return the Level's Lightmass Size as a FString */
-	FString GetLightmassSizeString() const;
+	UE_API FString GetLightmassSizeString() const;
 
 	/** @return the Level's File Size as a FString */
-	FString GetFileSizeString() const;
+	UE_API FString GetFileSizeString() const;
 	
 	/** @return Class used for streaming this level */
-	virtual UClass* GetStreamingClass() const;
+	UE_API virtual UClass* GetStreamingClass() const;
 
 protected:
 	/** Called when the map asset is renamed */
-	void OnAssetRenamed(const FAssetData& AssetData, const FString& OldObjectPath);
+	UE_API void OnAssetRenamed(const FAssetData& AssetData, const FString& OldObjectPath);
 
 protected:
 	/** Level model display name */
@@ -348,3 +350,5 @@ protected:
 	// Cached level actors count
 	int32								LevelActorsCount;
 };
+
+#undef UE_API
diff --git a/Engine/Source/Editor/WorldBrowser/Private/Tiles/WorldTileModel.h b/Engine/Source/Editor/WorldBrowser/Private/Tiles/WorldTileModel.h
--- a/Engine/Source/Editor/WorldBrowser/Private/Tiles/WorldTileModel.h
+++ b/Engine/Source/Editor/WorldBrowser/Private/Tiles/WorldTileModel.h
@@ -8,6 +8,8 @@
 #include "LevelModel.h"
 #include "LandscapeProxy.h"
 
+#define UE_API WORLDBROWSER_API
+
 struct FAssetData;
 class FLevelDragDropOp;
 class FWorldTileCollectionModel;
@@ -23,7 +25,7 @@ typedef TArray<TSharedPtr<class FWorldTileModel>> FWorldTileModelList;
 /**
  * The non-UI presentation logic for a single Level (Tile) in world composition
  */
-class WORLDBROWSER_API  FWorldTileModel
+class  FWorldTileModel
 	: public FLevelModel	
 {
 public:
@@ -71,145 +73,145 @@ public:
 	 *	@param	InWorldData		WorldBrowser world data
 	 *	@param	TileIdx			Tile index in world composition tiles list
 	 */
-	FWorldTileModel(FWorldTileCollectionModel& InWorldData, int32 TileIdx);
-	~FWorldTileModel();
+	UE_API FWorldTileModel(FWorldTileCollectionModel& InWorldData, int32 TileIdx);
+	UE_API ~FWorldTileModel();
 
 public:
 	// FLevelModel interface
-	virtual UObject* GetNodeObject() override;
-	virtual ULevel* GetLevelObject() const override;
-	virtual FName GetAssetName() const override;
-	virtual FName GetLongPackageName() const override;
-	virtual void Update() override;
-	virtual void UpdateAsset(const FAssetData& AssetData) override;
-	virtual void LoadLevel() override;
-	virtual void SetVisibleInEditor(bool bVisible) override;
-	virtual FVector2D GetLevelPosition2D() const override;
-	virtual FVector2D GetLevelSize2D() const override;
-	virtual void OnDrop(const TSharedPtr<FLevelDragDropOp>& Op) override;
-	virtual bool IsGoodToDrop(const TSharedPtr<FLevelDragDropOp>& Op) const override;
-	virtual void OnLevelAddedToWorld(ULevel* InLevel) override;
-	virtual void OnLevelRemovedFromWorld() override;
-	virtual void OnParentChanged() override;
-	virtual bool IsVisibleInCompositionView() const override;
-	virtual FLinearColor GetLevelColor() const override;
-	virtual void SetLevelColor(FLinearColor InColor) override;
+	UE_API virtual UObject* GetNodeObject() override;
+	UE_API virtual ULevel* GetLevelObject() const override;
+	UE_API virtual FName GetAssetName() const override;
+	UE_API virtual FName GetLongPackageName() const override;
+	UE_API virtual void Update() override;
+	UE_API virtual void UpdateAsset(const FAssetData& AssetData) override;
+	UE_API virtual void LoadLevel() override;
+	UE_API virtual void SetVisibleInEditor(bool bVisible) override;
+	UE_API virtual FVector2D GetLevelPosition2D() const override;
+	UE_API virtual FVector2D GetLevelSize2D() const override;
+	UE_API virtual void OnDrop(const TSharedPtr<FLevelDragDropOp>& Op) override;
+	UE_API virtual bool IsGoodToDrop(const TSharedPtr<FLevelDragDropOp>& Op) const override;
+	UE_API virtual void OnLevelAddedToWorld(ULevel* InLevel) override;
+	UE_API virtual void OnLevelRemovedFromWorld() override;
+	UE_API virtual void OnParentChanged() override;
+	UE_API virtual bool IsVisibleInCompositionView() const override;
+	UE_API virtual FLinearColor GetLevelColor() const override;
+	UE_API virtual void SetLevelColor(FLinearColor InColor) override;
 	// FLevelModel interface end
 	
 	/** Adds new streaming level*/
-	void AddStreamingLevel(UClass* InStreamingClass, const FName& InPackageName);
+	UE_API void AddStreamingLevel(UClass* InStreamingClass, const FName& InPackageName);
 
 	/** Assigns level to provided layer*/
-	void AssignToLayer(const FWorldTileLayer& InLayer);
+	UE_API void AssignToLayer(const FWorldTileLayer& InLayer);
 
 	/**	@return Whether tile is root of hierarchy */
-	bool IsRootTile() const;
+	UE_API bool IsRootTile() const;
 
 	/** Whether level should be visible in given area*/
-	bool ShouldBeVisible(FBox EditableArea) const;
+	UE_API bool ShouldBeVisible(FBox EditableArea) const;
 
 	/**	@return Whether level is shelved */
-	bool IsShelved() const;
+	UE_API bool IsShelved() const;
 
 	/** Hide a level from the editor */
-	void Shelve();
+	UE_API void Shelve();
 	
 	/** Show a level in the editor */
-	void Unshelve();
+	UE_API void Unshelve();
 
 	/** Whether this level landscape based or not */
-	bool IsLandscapeBased() const;
+	UE_API bool IsLandscapeBased() const;
 
 	UE_DEPRECATED(4.26, "No longer used; use CanReimportHeightmap instead.")
-	bool IsTiledLandscapeBased() const;
+	UE_API bool IsTiledLandscapeBased() const;
 
 	/** Whether this level has a heightmap that can be reimported or not */
-	bool CanReimportHeightmap() const;
+	UE_API bool CanReimportHeightmap() const;
 
 	UE_DEPRECATED(4.26, "No longer used; use IsLandscapeStreamingProxy instead.")
-	bool IsLandscapeProxy() const;
+	UE_API bool IsLandscapeProxy() const;
 
 	/** Whether this level has ALandscapeStreamingProxy or not */
-	bool IsLandscapeStreamingProxy() const;
+	UE_API bool IsLandscapeStreamingProxy() const;
 
 	/** @return The landscape actor in case this level is landscape based */
-	ALandscapeProxy* GetLandscape() const;
+	UE_API ALandscapeProxy* GetLandscape() const;
 	
 	/** Whether this level in provided layers list */
-	bool IsInLayersList(const TArray<FWorldTileLayer>& InLayerList) const;
+	UE_API bool IsInLayersList(const TArray<FWorldTileLayer>& InLayerList) const;
 	
 	/** @return Level position in shifted space */
-	FVector GetLevelCurrentPosition() const;
+	UE_API FVector GetLevelCurrentPosition() const;
 
 	/** @return Level relative position */
-	FIntVector GetRelativeLevelPosition() const;
+	UE_API FIntVector GetRelativeLevelPosition() const;
 	
 	/** @return Level absolute position in non shifted space */
-	FIntVector GetAbsoluteLevelPosition() const;
+	UE_API FIntVector GetAbsoluteLevelPosition() const;
 	
 	/** @return Calculates Level absolute position in non shifted space based on relative position */
-	FIntVector CalcAbsoluteLevelPosition() const;
+	UE_API FIntVector CalcAbsoluteLevelPosition() const;
 		
 	/** @return ULevel bounding box in shifted space*/
-	FBox GetLevelBounds() const override;
+	UE_API FBox GetLevelBounds() const override;
 	
 	/** Translate level center to new position */
-	void SetLevelPosition(const FIntVector& InPosition, const FIntPoint* InLandscapeSectionOffset = nullptr);
+	UE_API void SetLevelPosition(const FIntVector& InPosition, const FIntPoint* InLandscapeSectionOffset = nullptr);
 
 	/** Recursively sort all children by name */
-	void SortRecursive();
+	UE_API void SortRecursive();
 
 	/** 
 	 *	@return associated streaming level object for this tile
 	 *	Creates a new object in case it does not exists in a persistent world
 	 */
-	ULevelStreaming* GetAssociatedStreamingLevel();
+	UE_API ULevelStreaming* GetAssociatedStreamingLevel();
 
 	UE_DEPRECATED(4.20, "Call GetAssociatedStreamingLevel instead")
 	ULevelStreaming* GetAssosiatedStreamingLevel() { return GetAssociatedStreamingLevel(); }
 
 	/**  */
-	bool CreateAdjacentLandscapeProxy(ALandscapeProxy* SourceLandscape, FWorldTileModel::EWorldDirections InWhere);
+	UE_API bool CreateAdjacentLandscapeProxy(ALandscapeProxy* SourceLandscape, FWorldTileModel::EWorldDirections InWhere);
 
 	/**  */
-	ALandscapeProxy* ImportLandscapeTile(const FLandscapeImportSettings& Settings);
+	UE_API ALandscapeProxy* ImportLandscapeTile(const FLandscapeImportSettings& Settings);
 
 private:
 	/** Flush world info to package and level objects */
-	void OnLevelInfoUpdated();
+	UE_API void OnLevelInfoUpdated();
 
 	/** Fixup invalid streaming objects inside level */
-	void FixupStreamingObjects();
+	UE_API void FixupStreamingObjects();
 
 	/** When level with landscape is moved we need to update internal landscape coordinates to match landscape component grid  */
-	void UpdateLandscapeSectionsOffset(FIntPoint LevelOffset);
+	UE_API void UpdateLandscapeSectionsOffset(FIntPoint LevelOffset);
 	
 	/** Handler for LevelBoundsActorUpdated event */
-	void OnLevelBoundsActorUpdated();
+	UE_API void OnLevelBoundsActorUpdated();
 
 	/** Spawns AlevelBounds actor in the level in case it doesn't has one */
-	void EnsureLevelHasBoundsActor();
+	UE_API void EnsureLevelHasBoundsActor();
 
 	/** Handler for PostEditUndo  */
-	void OnPostUndoEvent();
+	UE_API void OnPostUndoEvent();
 	
 	/** Handler for PositionChanged event from Tile details object  */
-	void OnPositionPropertyChanged();
+	UE_API void OnPositionPropertyChanged();
 	
 	/** Handler for ParentPackageName event from Tile details object  */
-	void OnParentPackageNamePropertyChanged();
+	UE_API void OnParentPackageNamePropertyChanged();
 
 	/** Handler for LOD settings changes event from Tile details object  */
-	void OnLODSettingsPropertyChanged();
+	UE_API void OnLODSettingsPropertyChanged();
 	
 	/** Handler for ZOrder changes event from Tile details object  */
-	void OnZOrderPropertyChanged();
+	UE_API void OnZOrderPropertyChanged();
 
 	/** Handler for bHideInTileView changes event from Tile details object  */
-	void OnHideInTileViewChanged();
+	UE_API void OnHideInTileViewChanged();
 
 	/** Set the asset name based on the passed in package name */
-	void SetAssetName(const FName& PackageName);
+	UE_API void SetAssetName(const FName& PackageName);
 
 public:
 	/** This tile index in world composition tile list */
@@ -231,3 +233,5 @@ private:
 	/** Whether level has landscape components in it */
 	TWeakObjectPtr<ALandscapeProxy>			Landscape;
 };
+
+#undef UE_API
diff --git a/Engine/Source/Editor/WorldPartitionEditor/Public/WorldPartition/ContentBundle/ContentBundleEditorSubsystem.h b/Engine/Source/Editor/WorldPartitionEditor/Public/WorldPartition/ContentBundle/ContentBundleEditorSubsystem.h
--- a/Engine/Source/Editor/WorldPartitionEditor/Public/WorldPartition/ContentBundle/ContentBundleEditorSubsystem.h
+++ b/Engine/Source/Editor/WorldPartitionEditor/Public/WorldPartition/ContentBundle/ContentBundleEditorSubsystem.h
@@ -10,6 +10,8 @@
 #include "WorldPartition/ContentBundle/ContentBundleEditorSubsystemInterface.h"
 #include "ContentBundleEditorSubsystem.generated.h"
 
+#define UE_API WORLDPARTITIONEDITOR_API
+
 class FContentBundleEditor;
 class UContentBundleDescriptor;
 class UContentBundleEditorSubsystem;
@@ -102,8 +104,8 @@ private:
 	FOnActorEditorContextClientChanged ActorEditorContextClientChanged;
 };
 
-UCLASS()
-class WORLDPARTITIONEDITOR_API UContentBundleEditorSubsystem : public UEditorSubsystem, public IContentBundleEditorSubsystemInterface
+UCLASS(MinimalAPI)
+class UContentBundleEditorSubsystem : public UEditorSubsystem, public IContentBundleEditorSubsystemInterface
 {
 	GENERATED_BODY()
 
@@ -113,47 +115,49 @@ public:
 	UContentBundleEditorSubsystem() {}
 
 	//~ Begin UEditorSubsystem interface
-	virtual void Initialize(FSubsystemCollectionBase& Collection) override;
-	virtual void Deinitialize() override;
+	UE_API virtual void Initialize(FSubsystemCollectionBase& Collection) override;
+	UE_API virtual void Deinitialize() override;
 	//~ End UEditorSubsystem interface
 
 	//~ Begin IContentBundleEditorSubsystemInterface interface
 	virtual void NotifyContentBundleAdded(const FContentBundleEditor* ContentBundle) override { OnContentBundleAdded().Broadcast(ContentBundle); }
 	virtual void NotifyContentBundleRemoved(const FContentBundleEditor* ContentBundle) override { OnContentBundleRemoved().Broadcast(ContentBundle); }
-	virtual void NotifyContentBundleInjectedContent(const FContentBundleEditor* ContentBundle) override;
-	virtual void NotifyContentBundleRemovedContent(const FContentBundleEditor* ContentBundle) override;
+	UE_API virtual void NotifyContentBundleInjectedContent(const FContentBundleEditor* ContentBundle) override;
+	UE_API virtual void NotifyContentBundleRemovedContent(const FContentBundleEditor* ContentBundle) override;
 	virtual void NotifyContentBundleChanged(const FContentBundleEditor* ContentBundle) override { OnContentBundleChanged().Broadcast(ContentBundle); }
 	//~ End IContentBundleEditorSubsystemInterface interface
 
-	UWorld* GetWorld() const;
+	UE_API UWorld* GetWorld() const;
 
 	UContentBundleEditingSubmodule* GetEditingSubmodule() { return ContentBundleEditingSubModule; }
 
-	TSharedPtr<FContentBundleEditor> GetEditorContentBundleForActor(const AActor* Actor);
+	UE_API TSharedPtr<FContentBundleEditor> GetEditorContentBundleForActor(const AActor* Actor);
 
-	TArray<TSharedPtr<FContentBundleEditor>> GetEditorContentBundles();
-	TSharedPtr<FContentBundleEditor> GetEditorContentBundle(const UContentBundleDescriptor* ContentBundleDescriptor) const;
-	TSharedPtr<FContentBundleEditor> GetEditorContentBundle(const FGuid& ContentBundleGuid) const override;
+	UE_API TArray<TSharedPtr<FContentBundleEditor>> GetEditorContentBundles();
+	UE_API TSharedPtr<FContentBundleEditor> GetEditorContentBundle(const UContentBundleDescriptor* ContentBundleDescriptor) const;
+	UE_API TSharedPtr<FContentBundleEditor> GetEditorContentBundle(const FGuid& ContentBundleGuid) const override;
 
-	void SelectActors(FContentBundleEditor& EditorContentBundle);
-	void DeselectActors(FContentBundleEditor& EditorContentBundle);
+	UE_API void SelectActors(FContentBundleEditor& EditorContentBundle);
+	UE_API void DeselectActors(FContentBundleEditor& EditorContentBundle);
 
-	void ReferenceAllActors(FContentBundleEditor& EditorContentBundle);
-	void UnreferenceAllActors(FContentBundleEditor& EditorContentBundle);
+	UE_API void ReferenceAllActors(FContentBundleEditor& EditorContentBundle);
+	UE_API void UnreferenceAllActors(FContentBundleEditor& EditorContentBundle);
 
-	bool IsEditingContentBundle() const override;
-	bool IsEditingContentBundle(const FGuid& ContentBundleGuid) const;
-	bool ActivateContentBundleEditing(TSharedPtr<FContentBundleEditor>& ContentBundleEditor) const override;
-	bool DeactivateContentBundleEditing(TSharedPtr<FContentBundleEditor>& ContentBundleEditor) const override;
-	virtual bool DeactivateCurrentContentBundleEditing() const override;
-	bool IsContentBundleEditingActivated(TSharedPtr<FContentBundleEditor>& ContentBundleEditor) const;
+	UE_API bool IsEditingContentBundle() const override;
+	UE_API bool IsEditingContentBundle(const FGuid& ContentBundleGuid) const;
+	UE_API bool ActivateContentBundleEditing(TSharedPtr<FContentBundleEditor>& ContentBundleEditor) const override;
+	UE_API bool DeactivateContentBundleEditing(TSharedPtr<FContentBundleEditor>& ContentBundleEditor) const override;
+	UE_API virtual bool DeactivateCurrentContentBundleEditing() const override;
+	UE_API bool IsContentBundleEditingActivated(TSharedPtr<FContentBundleEditor>& ContentBundleEditor) const;
 
-	virtual void PushContentBundleEditing(bool bDuplicateContext) override;
-	virtual void PopContentBundleEditing() override;
+	UE_API virtual void PushContentBundleEditing(bool bDuplicateContext) override;
+	UE_API virtual void PopContentBundleEditing() override;
 
 private:
-	void SelectActorsInternal(FContentBundleEditor& EditorContentBundle, bool bSelect);
+	UE_API void SelectActorsInternal(FContentBundleEditor& EditorContentBundle, bool bSelect);
 
 	UPROPERTY()
 	TObjectPtr<UContentBundleEditingSubmodule> ContentBundleEditingSubModule;
-};
\ No newline at end of file
+};
+
+#undef UE_API
diff --git a/Engine/Source/Editor/WorldPartitionEditor/Public/WorldPartition/Filter/WorldPartitionActorFilterMode.h b/Engine/Source/Editor/WorldPartitionEditor/Public/WorldPartition/Filter/WorldPartitionActorFilterMode.h
--- a/Engine/Source/Editor/WorldPartitionEditor/Public/WorldPartition/Filter/WorldPartitionActorFilterMode.h
+++ b/Engine/Source/Editor/WorldPartitionEditor/Public/WorldPartition/Filter/WorldPartitionActorFilterMode.h
@@ -8,18 +8,20 @@
 #include "Containers/Map.h"
 #include "Templates/Tuple.h"
 
+#define UE_API WORLDPARTITIONEDITOR_API
+
 class SSceneOutliner;
 
-class WORLDPARTITIONEDITOR_API FWorldPartitionActorFilterMode : public ISceneOutlinerMode
+class FWorldPartitionActorFilterMode : public ISceneOutlinerMode
 {
 public:
 	// struct holding multi selection values that might be null if they differ between actors
-	class WORLDPARTITIONEDITOR_API FFilter
+	class FFilter
 	{
 	public:
-		FFilter(TSharedPtr<FWorldPartitionActorFilter> LevelFilter, const TArray<const FWorldPartitionActorFilter*>& SelectedFilters);
+		UE_API FFilter(TSharedPtr<FWorldPartitionActorFilter> LevelFilter, const TArray<const FWorldPartitionActorFilter*>& SelectedFilters);
 		
-		void Apply(FWorldPartitionActorFilter* Result) const;
+		UE_API void Apply(FWorldPartitionActorFilter* Result) const;
 				
 		struct FDataLayerFilter
 		{
@@ -32,9 +34,9 @@ public:
 
 		using FDataLayerFilters = TMap<FSoftObjectPath, FDataLayerFilter>;
 	private:
-		void Initialize(const FWorldPartitionActorFilter* Filter);
-		void Override(const FWorldPartitionActorFilter* Other);
-		void Merge(const FWorldPartitionActorFilter* Other);
+		UE_API void Initialize(const FWorldPartitionActorFilter* Filter);
+		UE_API void Override(const FWorldPartitionActorFilter* Other);
+		UE_API void Merge(const FWorldPartitionActorFilter* Other);
 				
 		// Unmodified Reference Level Filter
 		TSharedPtr<FWorldPartitionActorFilter> LevelFilter;
@@ -44,14 +46,14 @@ public:
 		friend class FWorldPartitionActorFilterMode;
 	};
 
-	FWorldPartitionActorFilterMode(SSceneOutliner* InSceneOutliner, TSharedPtr<FFilter> InFilter);
+	UE_API FWorldPartitionActorFilterMode(SSceneOutliner* InSceneOutliner, TSharedPtr<FFilter> InFilter);
 	
 	//~ Begin ISceneOutlinerMode interface
-	virtual void Rebuild() override;
-	virtual int32 GetTypeSortPriority(const ISceneOutlinerTreeItem& Item) const override;
+	UE_API virtual void Rebuild() override;
+	UE_API virtual int32 GetTypeSortPriority(const ISceneOutlinerTreeItem& Item) const override;
 	//~ Begin ISceneOutlinerMode interface
 		
-	void Apply(FWorldPartitionActorFilter& InOutResult) const;
+	UE_API void Apply(FWorldPartitionActorFilter& InOutResult) const;
 	
 	FFilter::FDataLayerFilters& FindChecked(const FWorldPartitionActorFilter* InFilter) const
 	{
@@ -65,9 +67,11 @@ public:
 
 protected:
 	//~ Begin ISceneOutlinerMode interface
-	virtual TUniquePtr<ISceneOutlinerHierarchy> CreateHierarchy() override;
+	UE_API virtual TUniquePtr<ISceneOutlinerHierarchy> CreateHierarchy() override;
 	//~ End ISceneOutlinerMode interface
 		
 private:
 	TSharedPtr<FFilter> Filter;
-};
\ No newline at end of file
+};
+
+#undef UE_API
diff --git a/Engine/Source/Editor/WorldPartitionEditor/Public/WorldPartition/Filter/WorldPartitionActorFilterPropertyTypeCustomization.h b/Engine/Source/Editor/WorldPartitionEditor/Public/WorldPartition/Filter/WorldPartitionActorFilterPropertyTypeCustomization.h
--- a/Engine/Source/Editor/WorldPartitionEditor/Public/WorldPartition/Filter/WorldPartitionActorFilterPropertyTypeCustomization.h
+++ b/Engine/Source/Editor/WorldPartitionEditor/Public/WorldPartition/Filter/WorldPartitionActorFilterPropertyTypeCustomization.h
@@ -8,18 +8,20 @@
 #include "Templates/SharedPointer.h"
 #include "WorldPartition/Filter/WorldPartitionActorFilterMode.h"
 
+#define UE_API WORLDPARTITIONEDITOR_API
+
 class FDetailWidgetRow;
 class IDetailChildrenBuilder;
 class IPropertyHandle;
 class ISceneOutliner;
 
 // Registered (FLevelInstanceEditorModule::StartupModule) Property Customization for properties of type FWorldPartitionActorFilter
-struct WORLDPARTITIONEDITOR_API FWorldPartitionActorFilterPropertyTypeCustomization : public IPropertyTypeCustomization
+struct FWorldPartitionActorFilterPropertyTypeCustomization : public IPropertyTypeCustomization
 {
 public:
-	~FWorldPartitionActorFilterPropertyTypeCustomization();
+	UE_API ~FWorldPartitionActorFilterPropertyTypeCustomization();
 
-	virtual void CustomizeHeader(TSharedRef<IPropertyHandle> StructPropertyHandle, FDetailWidgetRow& HeaderRow, IPropertyTypeCustomizationUtils& StructCustomizationUtils) override;
+	UE_API virtual void CustomizeHeader(TSharedRef<IPropertyHandle> StructPropertyHandle, FDetailWidgetRow& HeaderRow, IPropertyTypeCustomizationUtils& StructCustomizationUtils) override;
 	virtual void CustomizeChildren(TSharedRef<IPropertyHandle> StructPropertyHandle, IDetailChildrenBuilder& StructBuilder, IPropertyTypeCustomizationUtils& StructCustomizationUtils) override {}
 
 private:
@@ -29,3 +31,5 @@ private:
 	FDelegateHandle WorldPartitionActorFilterChangedHandle;
 	TSharedPtr<ISceneOutliner> SceneOutliner;
 };
+
+#undef UE_API
diff --git a/Engine/Source/Editor/WorldPartitionEditor/Public/WorldPartition/HLOD/HLODEditorSubsystem.h b/Engine/Source/Editor/WorldPartitionEditor/Public/WorldPartition/HLOD/HLODEditorSubsystem.h
--- a/Engine/Source/Editor/WorldPartitionEditor/Public/WorldPartition/HLOD/HLODEditorSubsystem.h
+++ b/Engine/Source/Editor/WorldPartitionEditor/Public/WorldPartition/HLOD/HLODEditorSubsystem.h
@@ -11,6 +11,8 @@
 
 #include "HLODEditorSubsystem.generated.h"
 
+#define UE_API WORLDPARTITIONEDITOR_API
+
 
 class AActor;
 class AWorldPartitionHLOD;
@@ -31,53 +33,53 @@ enum class EHLODSettingsVisibility : uint8
 /**
  * UWorldPartitionHLODEditorSubsystem
  */
-UCLASS()
-class WORLDPARTITIONEDITOR_API UWorldPartitionHLODEditorSubsystem : public UTickableWorldSubsystem
+UCLASS(MinimalAPI)
+class UWorldPartitionHLODEditorSubsystem : public UTickableWorldSubsystem
 {
 	GENERATED_BODY()
 
 public:
-	UWorldPartitionHLODEditorSubsystem();
-	virtual ~UWorldPartitionHLODEditorSubsystem();
+	UE_API UWorldPartitionHLODEditorSubsystem();
+	UE_API virtual ~UWorldPartitionHLODEditorSubsystem();
 
 	//~ Begin USubsystem Interface.
-	virtual void Initialize(FSubsystemCollectionBase& Collection) override;
-	virtual void Deinitialize() override;
+	UE_API virtual void Initialize(FSubsystemCollectionBase& Collection) override;
+	UE_API virtual void Deinitialize() override;
 	//~ End USubsystem Interface.
 
 	//~ Begin UWorldSubsystem Interface.
-	virtual bool DoesSupportWorldType(const EWorldType::Type WorldType) const override;
+	UE_API virtual bool DoesSupportWorldType(const EWorldType::Type WorldType) const override;
 	//~ End UWorldSubsystem Interface.
 
 	//~ Begin FTickableGameObject Interface
 	virtual bool IsTickable() const override { return true; }
 	virtual bool IsTickableInEditor() const override { return true; }
-	virtual void Tick(float DeltaTime) override;
-	virtual TStatId GetStatId() const override;
+	UE_API virtual void Tick(float DeltaTime) override;
+	UE_API virtual TStatId GetStatId() const override;
 	//~ End FTickableGameObject Interface
 
-	virtual bool WriteHLODStats(const IWorldPartitionEditorModule::FWriteHLODStatsParams& Params) const;
+	UE_API virtual bool WriteHLODStats(const IWorldPartitionEditorModule::FWriteHLODStatsParams& Params) const;
 
-	static void AddHLODSettingsFilter(EHLODSettingsVisibility InSettingsVisibility, TSoftObjectPtr<UStruct> InStruct, FName InPropertyName);
+	static UE_API void AddHLODSettingsFilter(EHLODSettingsVisibility InSettingsVisibility, TSoftObjectPtr<UStruct> InStruct, FName InPropertyName);
 	
 private:
-	bool IsHLODInEditorEnabled();
-	void SetHLODInEditorEnabled(bool bInEnable);
+	UE_API bool IsHLODInEditorEnabled();
+	UE_API void SetHLODInEditorEnabled(bool bInEnable);
 
-	void OnWorldPartitionInitialized(UWorldPartition* InWorldPartition);
-	void OnWorldPartitionUninitialized(UWorldPartition* InWorldPartition);
+	UE_API void OnWorldPartitionInitialized(UWorldPartition* InWorldPartition);
+	UE_API void OnWorldPartitionUninitialized(UWorldPartition* InWorldPartition);
 
-	void OnLoaderAdapterStateChanged(const IWorldPartitionActorLoaderInterface::ILoaderAdapter* LoaderAdapter);
+	UE_API void OnLoaderAdapterStateChanged(const IWorldPartitionActorLoaderInterface::ILoaderAdapter* LoaderAdapter);
 
-	void ForceHLODStateUpdate();
+	UE_API void ForceHLODStateUpdate();
 
-	bool WriteHLODStats(const FString& InFilename) const;
-	bool WriteHLODInputStats(const FString& InFilename) const;
+	UE_API bool WriteHLODStats(const FString& InFilename) const;
+	UE_API bool WriteHLODInputStats(const FString& InFilename) const;
 	
-	void OnWorldPartitionEditorSettingsChanged(const FName& PropertyName, const UWorldPartitionEditorSettings& WorldPartitionEditorSettings);
-	void ApplyHLODSettingsFiltering();
+	UE_API void OnWorldPartitionEditorSettingsChanged(const FName& PropertyName, const UWorldPartitionEditorSettings& WorldPartitionEditorSettings);
+	UE_API void ApplyHLODSettingsFiltering();
 
-	void OnColorHandlerPropertyChangedEvent(UObject* InObject, FPropertyChangedEvent& InPropertyChangedEvent);
+	UE_API void OnColorHandlerPropertyChangedEvent(UObject* InObject, FPropertyChangedEvent& InPropertyChangedEvent);
 
 private:
 	FVector CachedCameraLocation;
@@ -89,7 +91,7 @@ private:
 	TMap<TObjectKey<UWorldPartition>, TPimplPtr<FWorldPartitionHLODEditorData>> WorldPartitionsHLODEditorData;
 
 	typedef TMap<TSoftObjectPtr<UStruct>, TSet<FName>> FStructsPropertiesMap;
-	static TMap<EHLODSettingsVisibility, FStructsPropertiesMap> StructsPropertiesVisibility;
+	static UE_API TMap<EHLODSettingsVisibility, FStructsPropertiesMap> StructsPropertiesVisibility;
 };
 
 
@@ -104,4 +106,6 @@ private:
 	HLOD_ADD_CLASS_SETTING_FILTER_NAME(SettingsLevel, TypeIdentifier, GET_MEMBER_NAME_CHECKED(TypeIdentifier, PropertyIdentifier))
 
 #define HLOD_ADD_STRUCT_SETTING_FILTER(SettingsLevel, TypeIdentifier, PropertyIdentifier) \
-	HLOD_ADD_STRUCT_SETTING_FILTER_NAME(SettingsLevel, TypeIdentifier, GET_MEMBER_NAME_CHECKED(TypeIdentifier, PropertyIdentifier))
\ No newline at end of file
+	HLOD_ADD_STRUCT_SETTING_FILTER_NAME(SettingsLevel, TypeIdentifier, GET_MEMBER_NAME_CHECKED(TypeIdentifier, PropertyIdentifier))
+
+#undef UE_API
diff --git a/Engine/Source/Editor/WorldPartitionEditor/Public/WorldPartition/HLOD/HLODLayerFactory.h b/Engine/Source/Editor/WorldPartitionEditor/Public/WorldPartition/HLOD/HLODLayerFactory.h
--- a/Engine/Source/Editor/WorldPartitionEditor/Public/WorldPartition/HLOD/HLODLayerFactory.h
+++ b/Engine/Source/Editor/WorldPartitionEditor/Public/WorldPartition/HLOD/HLODLayerFactory.h
@@ -8,16 +8,20 @@
 
 #include "HLODLayerFactory.generated.h"
 
+#define UE_API WORLDPARTITIONEDITOR_API
+
 class FFeedbackContext;
 class UClass;
 class UObject;
 
-UCLASS()
-class WORLDPARTITIONEDITOR_API UHLODLayerFactory : public UFactory
+UCLASS(MinimalAPI)
+class UHLODLayerFactory : public UFactory
 {
 	GENERATED_UCLASS_BODY()
 
 	//~ Begin UFactory Interface
-	virtual UObject* FactoryCreateNew(UClass* Class, UObject* InParent, FName Name, EObjectFlags Flags, UObject* Context, FFeedbackContext* Warn) override;
+	UE_API virtual UObject* FactoryCreateNew(UClass* Class, UObject* InParent, FName Name, EObjectFlags Flags, UObject* Context, FFeedbackContext* Warn) override;
 	//~ Begin UFactory Interface	
 };
+
+#undef UE_API
diff --git a/Engine/Source/Editor/WorldPartitionEditor/Public/WorldPartition/WorldPartitionEditorSettings.h b/Engine/Source/Editor/WorldPartitionEditor/Public/WorldPartition/WorldPartitionEditorSettings.h
--- a/Engine/Source/Editor/WorldPartitionEditor/Public/WorldPartition/WorldPartitionEditorSettings.h
+++ b/Engine/Source/Editor/WorldPartitionEditor/Public/WorldPartition/WorldPartitionEditorSettings.h
@@ -7,58 +7,60 @@
 #include "Commandlets/WorldPartitionConvertCommandlet.h"
 #include "WorldPartitionEditorSettings.generated.h"
 
-UCLASS(config = EditorSettings, meta = (DisplayName = "World Partition"))
-class WORLDPARTITIONEDITOR_API UWorldPartitionEditorSettings : public UDeveloperSettings
+#define UE_API WORLDPARTITIONEDITOR_API
+
+UCLASS(MinimalAPI, config = EditorSettings, meta = (DisplayName = "World Partition"))
+class UWorldPartitionEditorSettings : public UDeveloperSettings
 {
 	GENERATED_BODY()
 
 public:
-	UWorldPartitionEditorSettings();
+	UE_API UWorldPartitionEditorSettings();
 
 	DECLARE_MULTICAST_DELEGATE_TwoParams(FOnWorldPartitionEditorSettingsChanged, const FName&, const UWorldPartitionEditorSettings&);
-	static FOnWorldPartitionEditorSettingsChanged& OnSettingsChanged();
+	static UE_API FOnWorldPartitionEditorSettingsChanged& OnSettingsChanged();
 
-	static FName GetEnableLoadingInEditorPropertyName();
-	static FName GetEnableAdvancedHLODSettingsPropertyName();
+	static UE_API FName GetEnableLoadingInEditorPropertyName();
+	static UE_API FName GetEnableAdvancedHLODSettingsPropertyName();
 
-	TSubclassOf<UWorldPartitionConvertCommandlet> GetCommandletClass() const;
-	void SetCommandletClass(const TSubclassOf<UWorldPartitionConvertCommandlet>& InCommandletClass);
+	UE_API TSubclassOf<UWorldPartitionConvertCommandlet> GetCommandletClass() const;
+	UE_API void SetCommandletClass(const TSubclassOf<UWorldPartitionConvertCommandlet>& InCommandletClass);
 
-	int32 GetInstancedFoliageGridSize() const;
-	void SetInstancedFoliageGridSize(int32 InInstancedFoliageGridSize);
+	UE_API int32 GetInstancedFoliageGridSize() const;
+	UE_API void SetInstancedFoliageGridSize(int32 InInstancedFoliageGridSize);
 
-	int32 GetMinimapLowQualityWorldUnitsPerPixelThreshold() const;
-	void SetMinimapLowQualityWorldUnitsPerPixelThreshold(int32 InMinimapLowQualityWorldUnitsPerPixelThreshold);
+	UE_API int32 GetMinimapLowQualityWorldUnitsPerPixelThreshold() const;
+	UE_API void SetMinimapLowQualityWorldUnitsPerPixelThreshold(int32 InMinimapLowQualityWorldUnitsPerPixelThreshold);
 
-	bool GetEnableLoadingInEditor() const;
-	void SetEnableLoadingInEditor(bool bInEnableLoadingInEditor);
+	UE_API bool GetEnableLoadingInEditor() const;
+	UE_API void SetEnableLoadingInEditor(bool bInEnableLoadingInEditor);
 
-	bool GetEnableStreamingGenerationLogOnPIE() const;
-	void SetEnableStreamingGenerationLogOnPIE(bool bInEnableStreamingGenerationLogOnPIE);
+	UE_API bool GetEnableStreamingGenerationLogOnPIE() const;
+	UE_API void SetEnableStreamingGenerationLogOnPIE(bool bInEnableStreamingGenerationLogOnPIE);
 
-	bool GetShowHLODsInEditor() const;
-	void SetShowHLODsInEditor(bool bInShowHLODsInEditor);
+	UE_API bool GetShowHLODsInEditor() const;
+	UE_API void SetShowHLODsInEditor(bool bInShowHLODsInEditor);
 
-	bool GetShowHLODsOverLoadedRegions() const;
-	void SetShowHLODsOverLoadedRegions(bool bInShowHLODsOverLoadedRegions);
+	UE_API bool GetShowHLODsOverLoadedRegions() const;
+	UE_API void SetShowHLODsOverLoadedRegions(bool bInShowHLODsOverLoadedRegions);
 
-	bool GetEnableAdvancedHLODSettings() const;
-	void SetEnableAdvancedHLODSettings(bool bInEnableAdvancedHLODSettings);
+	UE_API bool GetEnableAdvancedHLODSettings() const;
+	UE_API void SetEnableAdvancedHLODSettings(bool bInEnableAdvancedHLODSettings);
 
-	double GetHLODMinDrawDistance() const;
-	void SetHLODMinDrawDistance(double bInHLODMinDrawDistance);
+	UE_API double GetHLODMinDrawDistance() const;
+	UE_API void SetHLODMinDrawDistance(double bInHLODMinDrawDistance);
 
-	double GetHLODMaxDrawDistance() const;
-	void SetHLODMaxDrawDistance(double bInHLODMaxDrawDistance);
+	UE_API double GetHLODMaxDrawDistance() const;
+	UE_API void SetHLODMaxDrawDistance(double bInHLODMaxDrawDistance);
 
-	bool GetDisableBugIt() const;
-	void SetDisableBugIt(bool bInEnableBugIt);
+	UE_API bool GetDisableBugIt() const;
+	UE_API void SetDisableBugIt(bool bInEnableBugIt);
 
-	bool GetDisablePIE() const;
-	void SetDisablePIE(bool bInEnablePIE);
+	UE_API bool GetDisablePIE() const;
+	UE_API void SetDisablePIE(bool bInEnablePIE);
 
-	bool GetAdvancedMode() const;
-	void SetAdvancedMode(bool bInAdvancedMode);
+	UE_API bool GetAdvancedMode() const;
+	UE_API void SetAdvancedMode(bool bInAdvancedMode);
 
 public:
 	UE_DEPRECATED(5.5, "Use Get/SetCommandletClass()")
@@ -113,7 +115,9 @@ public:
 protected:
 	FOnWorldPartitionEditorSettingsChanged SettingsChangedDelegate;
 
-	void TriggerPropertyChangeEvent(FName PropertyName);
-	virtual void PostEditChangeProperty(FPropertyChangedEvent& PropertyChangedEvent) override;
+	UE_API void TriggerPropertyChangeEvent(FName PropertyName);
+	UE_API virtual void PostEditChangeProperty(FPropertyChangedEvent& PropertyChangedEvent) override;
+
+};
 
-};
\ No newline at end of file
+#undef UE_API
diff --git a/Engine/Source/Editor/ZenEditor/Private/ZenDashboardLauncher.h b/Engine/Source/Editor/ZenEditor/Private/ZenDashboardLauncher.h
--- a/Engine/Source/Editor/ZenEditor/Private/ZenDashboardLauncher.h
+++ b/Engine/Source/Editor/ZenEditor/Private/ZenDashboardLauncher.h
@@ -13,7 +13,7 @@
 namespace UE::Zen
 {
 
-class ZENEDITOR_API FZenDashboardLauncher : public TSharedFromThis<FZenDashboardLauncher>
+class FZenDashboardLauncher : public TSharedFromThis<FZenDashboardLauncher>
 {
 public:
 	enum class EStartZenDashboardResult : uint32
@@ -28,7 +28,7 @@ public:
 	UE_API FZenDashboardLauncher();
 	UE_API ~FZenDashboardLauncher();
 	
-	static UE_API const TSharedPtr<FZenDashboardLauncher>& Get()
+	static const TSharedPtr<FZenDashboardLauncher>& Get()
 	{
 		if (!Instance)
 		{
@@ -69,13 +69,13 @@ private:
 	 * Attempts building ZenDashboard via UAT, will launch with forwarded parameters if successful.
 	 * Assumes that the ZenDashboard Executable in path belongs to this engine.
 	 */
-	void BuildZenDashboard(const FString& Path, const FString& LaunchParameters, StartZenDashboardCallback Callback);
+	UE_API void BuildZenDashboard(const FString& Path, const FString& LaunchParameters, StartZenDashboardCallback Callback);
 
 private:
 	/** The process handler of ZenDashboard. */
 	FProcHandle ZenDashboardHandle;
 
-	static TSharedPtr<FZenDashboardLauncher> Instance;
+	static UE_API TSharedPtr<FZenDashboardLauncher> Instance;
 };
 
 } // namespace UE::Zen
diff --git a/Engine/Source/Editor/ZenEditor/Public/ZenEditor.h b/Engine/Source/Editor/ZenEditor/Public/ZenEditor.h
--- a/Engine/Source/Editor/ZenEditor/Public/ZenEditor.h
+++ b/Engine/Source/Editor/ZenEditor/Public/ZenEditor.h
@@ -9,6 +9,8 @@
 #include "Templates/UniquePtr.h"
 #include "IDerivedDataCacheNotifications.h"
 
+#define UE_API ZENEDITOR_API
+
 class FSpawnTabArgs;
 class SDerivedDataCacheSettingsDialog;
 class IDerivedDataCacheNotifications;
@@ -19,39 +21,39 @@ class SWindow;
 /**
  * The module holding all of the UI related pieces for DerivedData
  */
-class ZENEDITOR_API FZenEditor : public IModuleInterface
+class FZenEditor : public IModuleInterface
 {
 public:
 
 	/**
 	 * Called right after the module DLL has been loaded and the module object has been created
 	 */
-	virtual void StartupModule() override;
+	UE_API virtual void StartupModule() override;
 
 	/**
 	 * Called before the module is unloaded, right before the module object is destroyed.
 	 */
-	virtual void ShutdownModule() override;
+	UE_API virtual void ShutdownModule() override;
 
-	bool IsZenEnabled() const;
-	TSharedRef<SWidget>	CreateStatusBarWidget();
+	UE_API bool IsZenEnabled() const;
+	UE_API TSharedRef<SWidget>	CreateStatusBarWidget();
 
-	void ShowResourceUsageTab();
-	void ShowCacheStatisticsTab();
-	void ShowZenServerStatusTab();
-	void StartZenServer();
-	void StopZenServer();
-	void RestartZenServer();
+	UE_API void ShowResourceUsageTab();
+	UE_API void ShowCacheStatisticsTab();
+	UE_API void ShowZenServerStatusTab();
+	UE_API void StartZenServer();
+	UE_API void StopZenServer();
+	UE_API void RestartZenServer();
 
 private:
 
-	TSharedPtr<SWidget> CreateResourceUsageDialog();
-	TSharedPtr<SWidget> CreateCacheStatisticsDialog();
-	TSharedPtr<SWidget> CreateZenStoreDialog();
+	UE_API TSharedPtr<SWidget> CreateResourceUsageDialog();
+	UE_API TSharedPtr<SWidget> CreateCacheStatisticsDialog();
+	UE_API TSharedPtr<SWidget> CreateZenStoreDialog();
 
-	TSharedRef<SDockTab> CreateResourceUsageTab(const FSpawnTabArgs& Args);
-	TSharedRef<SDockTab> CreateCacheStatisticsTab(const FSpawnTabArgs& Args);
-	TSharedRef<SDockTab> CreateZenServerStatusTab(const FSpawnTabArgs& Args);
+	UE_API TSharedRef<SDockTab> CreateResourceUsageTab(const FSpawnTabArgs& Args);
+	UE_API TSharedRef<SDockTab> CreateCacheStatisticsTab(const FSpawnTabArgs& Args);
+	UE_API TSharedRef<SDockTab> CreateZenServerStatusTab(const FSpawnTabArgs& Args);
 
 	TWeakPtr<SDockTab> ResourceUsageTab;
 	TWeakPtr<SDockTab> CacheStatisticsTab;
@@ -63,3 +65,4 @@ private:
 };
 
 
+#undef UE_API
diff --git a/Engine/Source/Runtime/AugmentedReality/Public/ARActor.h b/Engine/Source/Runtime/AugmentedReality/Public/ARActor.h
--- a/Engine/Source/Runtime/AugmentedReality/Public/ARActor.h
+++ b/Engine/Source/Runtime/AugmentedReality/Public/ARActor.h
@@ -7,32 +7,34 @@
 #include "Engine/GameEngine.h"
 #include "ARActor.generated.h"
 
+#define UE_API AUGMENTEDREALITY_API
+
 class UARComponent;
 
-UCLASS(BlueprintType, Experimental, Category="AR Gameplay")
-class AUGMENTEDREALITY_API AARActor: public AActor
+UCLASS(MinimalAPI, BlueprintType, Experimental, Category="AR Gameplay")
+class AARActor: public AActor
 {
 	GENERATED_BODY()
 	
 public:
-	AARActor();
+	UE_API AARActor();
 	
 	UFUNCTION(BlueprintCallable, Category="AR Gameplay")
-	UARComponent* AddARComponent(TSubclassOf<UARComponent> InComponentClass, const FGuid& NativeID);
+	UE_API UARComponent* AddARComponent(TSubclassOf<UARComponent> InComponentClass, const FGuid& NativeID);
 	
-	static void RequestSpawnARActor(FGuid NativeID, UClass* InComponentClass);
-	static void RequestDestroyARActor(AARActor* InActor);
+	static UE_API void RequestSpawnARActor(FGuid NativeID, UClass* InComponentClass);
+	static UE_API void RequestDestroyARActor(AARActor* InActor);
 };
 
 USTRUCT()
-struct AUGMENTEDREALITY_API FTrackedGeometryGroup
+struct FTrackedGeometryGroup
 {
 public:
 	GENERATED_USTRUCT_BODY()
 	
 	FTrackedGeometryGroup() = default;
 	
-	FTrackedGeometryGroup(UARTrackedGeometry* InTrackedGeometry);
+	UE_API FTrackedGeometryGroup(UARTrackedGeometry* InTrackedGeometry);
 	
 	UPROPERTY()
 	TObjectPtr<AARActor> ARActor = nullptr;
@@ -43,3 +45,5 @@ public:
 	UPROPERTY()
 	TObjectPtr<UARTrackedGeometry> TrackedGeometry = nullptr;
 };
+
+#undef UE_API
